Bink Audio :: 2/4 :: Decoder :: #pf 60aea2116905a6000136901e #pf 60c8e2ae86ce760001d99596

Unreal part of the decoder, linking against the SDK from checkin 1/4. This code handles the file/block parsing, as well as seek table access.

diff --git a/Engine/Source/Runtime/BinkAudioDecoder/BinkAudioDecoder.Build.cs b/Engine/Source/Runtime/BinkAudioDecoder/BinkAudioDecoder.Build.cs
new file mode 100644
--- /dev/null
+++ b/Engine/Source/Runtime/BinkAudioDecoder/BinkAudioDecoder.Build.cs
@@ -0,0 +1,64 @@
+// Copyright Epic Games, Inc. All Rights Reserved.
+
+using System.IO;
+using System.Collections.Generic;
+using System.Security.Cryptography;
+using UnrealBuildTool;
+
+public class BinkAudioDecoder : ModuleRules
+{
+    // virtual so that NDA platforms can hide secrets
+    protected virtual string GetLibrary(ReadOnlyTargetRules Target)
+    {
+        if (Target.Platform == UnrealTargetPlatform.Win64)
+        {
+            return Path.Combine(Target.UEThirdPartySourceDirectory, "BinkAudio", "Lib", "binka_ue_decode_win64_static.lib");
+        }
+        if (Target.Platform == UnrealTargetPlatform.Linux)
+        {
+            return Path.Combine(Target.UEThirdPartySourceDirectory, "BinkAudio", "Lib", "libbinka_ue_decode_lnx64_static.a");
+        }
+        if (Target.Platform == UnrealTargetPlatform.Mac)
+        {
+            return Path.Combine(Target.UEThirdPartySourceDirectory, "BinkAudio", "Lib", "libbinka_ue_decode_osx64_static.a");
+        }
+        if (Target.Platform == UnrealTargetPlatform.IOS)
+        {
+            return Path.Combine(Target.UEThirdPartySourceDirectory, "BinkAudio", "Lib", "libbinka_ue_decode_ios_static.a");
+        }
+        if (Target.Platform == UnrealTargetPlatform.Android)
+        {
+            return Path.Combine(Target.UEThirdPartySourceDirectory, "BinkAudio", "Lib", "libbinka_ue_decode_androidarm64_static.a");
+        }        
+        return null;
+    }
+
+    public BinkAudioDecoder(ReadOnlyTargetRules Target) : base(Target)
+    {
+        ShortName = "BinkAudioDecoder";
+
+        PublicDependencyModuleNames.AddRange(
+            new string[]
+            {
+                "Core",
+                "CoreUObject",
+                "Engine"
+            }
+        );
+
+        PublicSystemIncludePaths.Add(Path.Combine(Target.UEThirdPartySourceDirectory, "BinkAudio", "Include"));
+
+        string Library = GetLibrary(Target);
+        if (Library != null)
+        {
+			// We have to have this because some audio mixers are shared across platforms
+			// that we don't have libs for (e.g. hololens).
+            PublicDefinitions.Add("WITH_BINK_AUDIO=1");
+            PublicAdditionalLibraries.Add(Library);
+        }
+        else
+        {
+            PublicDefinitions.Add("WITH_BINK_AUDIO=0");
+        }
+    }
+}
diff --git a/Engine/Source/Runtime/BinkAudioDecoder/Private/BinkAudioInfo.cpp b/Engine/Source/Runtime/BinkAudioDecoder/Private/BinkAudioInfo.cpp
new file mode 100644
--- /dev/null
+++ b/Engine/Source/Runtime/BinkAudioDecoder/Private/BinkAudioInfo.cpp
@@ -0,0 +1,645 @@
+// Copyright Epic Games, Inc. All Rights Reserved.
+
+#include "BinkAudioInfo.h"
+#include "Interfaces/IAudioFormat.h"
+
+#include "Modules/ModuleInterface.h"
+
+#include "binka_ue_file_header.h"
+#include "binka_ue_decode.h"
+
+#if !defined(PLATFORM_LITTLE_ENDIAN) || !PLATFORM_LITTLE_ENDIAN
+#error Bink Audio hasn't been updated for big endian.
+#endif
+
+DEFINE_LOG_CATEGORY_STATIC(LogBinkAudioDecoder, Log, All);
+
+//-----------------------------------------------------------------------------
+//-----------------------------------------------------------------------------
+class BINKAUDIODECODER_API FBinkAudioDecoderModule : public IModuleInterface
+{
+public:
+	virtual void StartupModule() override {}
+	virtual void ShutdownModule() override {}
+};
+
+#define PTR_ADD(ptr,off) ((void*)(((uint8*)(ptr))+(off)))
+#define Align32( val ) ( ( ( val ) + 31 ) & ~31 )
+
+
+// This is shared in the encoder as the max number of bink audio streams. If
+// you want to rebuild that, this can be otherwise be changed easily though
+// it needs to be < 255
+#define MAX_BINK_AUDIO_CHANNELS 16
+
+//-----------------------------------------------------------------------------
+//
+// All memory for the decoder is in one contiguous block:
+//
+// BinkAudioDecoder structure
+// StreamChannels
+// Decoders
+// OutputReservoir
+// DeinterlaceBuffer
+// SeekTable[SeekTableCount + 1] // see comment below
+//
+//-----------------------------------------------------------------------------
+struct BinkAudioDecoder
+{
+	// # of low level bink audio streams (max 2 chans each)
+	uint8 StreamCount;
+	// # of current bytes in the res
+	uint32 OutputReservoirValidBytes;
+	// # max size of the res.
+	uint32 OutputReservoirTotalBytes;
+	// offsets to the various pieces of the decoder
+	uint16 ToDeinterlaceBufferOffset32;
+	uint16 ToSeekTableOffset32;
+	uint16 ToOutputReservoirOffset32;
+	// # of entries in the seek table - the entries are byte sizes, so 
+	// when decoding the seek table we convert to file offsets. This means
+	// we actually have +1 entries in the decoded seek table due to the standard
+	// "span -> edge" count thing.
+	uint16 SeekTableCount;
+	// # of low level bink blocks one seek table entry spans
+	uint16 FramesPerSeekTableEntry;
+	// # of frames we need to eat before outputting data due to a sample
+	// accurate seek.
+	uint16 ConsumeFrameCount;
+
+	uint8* StreamChannels()
+	{
+		return (uint8*)(this + 1);
+	}
+	void** Decoders()
+	{
+		return (void**)PTR_ADD(this, Align32(sizeof(uint8) * StreamCount + sizeof(BinkAudioDecoder)));
+	}
+	uint8* OutputReservoir()
+	{
+		return (uint8*)(PTR_ADD(this, ToOutputReservoirOffset32 * 32U));
+	}
+	int16* DeinterlaceBuffer()
+	{
+		return (int16*)(PTR_ADD(this, ToDeinterlaceBufferOffset32 * 32U));
+	}
+	uint32* SeekTable() // [SeekTableCount + 1] if SeekTableCount != 0, otherwise invalid.
+	{
+		return (uint32*)(PTR_ADD(this, ToSeekTableOffset32 * 32U));
+	}
+};
+
+//-----------------------------------------------------------------------------
+//-----------------------------------------------------------------------------
+FBinkAudioInfo::FBinkAudioInfo()
+{
+}
+
+//-----------------------------------------------------------------------------
+//-----------------------------------------------------------------------------
+FBinkAudioInfo::~FBinkAudioInfo()
+{
+	FMemory::Free(RawMemory);
+}
+
+//-----------------------------------------------------------------------------
+//-----------------------------------------------------------------------------
+void FBinkAudioInfo::SeekToTime(const float SeekTimeSeconds)
+{
+	if (Decoder->SeekTableCount == 0)
+	{
+		// Can't seek without a seek table.
+		return;
+	}
+
+	uint32 SeekTimeSamples = 0;
+	if (SeekTimeSeconds > 0)
+	{
+		SeekTimeSamples = (uint32)(SeekTimeSeconds * SampleRate);
+	}
+
+	uint32 SamplesInFrame = GetMaxFrameSizeSamples();
+	if (SeekTimeSamples > this->TrueSampleCount)
+	{
+		SeekTimeSamples = this->TrueSampleCount - 1;
+	}
+
+	uint32 SamplesPerBlock = SamplesInFrame * Decoder->FramesPerSeekTableEntry;
+	uint32 SeekTableIndex = SeekTimeSamples / SamplesPerBlock;
+	uint32 SeekTableOffset = SeekTimeSamples % SamplesPerBlock;
+
+	uint32 OffsetToBlock = Decoder->SeekTable()[SeekTableIndex] + sizeof(BinkAudioFileHeader) + Decoder->SeekTableCount * sizeof(uint16);
+
+	Decoder->ConsumeFrameCount = SeekTableOffset;
+	Decoder->OutputReservoirValidBytes = 0;
+
+#if WITH_BINK_AUDIO
+	UEBinkAudioDecodeInterface* BinkInterface = UnrealBinkAudioDecodeInterface();
+	void** Streams = Decoder->Decoders();
+	for (uint8 i = 0; i < Decoder->StreamCount; i++)
+	{
+		BinkInterface->ResetStartFn(Streams[i]);
+	}
+#endif // WITH_BINK_AUDIO
+
+	//
+	// Here we need to set up the data we get to point at the right spot.
+	//
+	if (StreamingSoundWave == nullptr)
+	{
+		// If we aren't streaming we can just go directly to the offset we need.
+		this->SrcBufferOffset = OffsetToBlock;
+	}
+	else
+	{
+		const uint32 TotalStreamingChunks = StreamingSoundWave->GetNumChunks();
+		uint32 RemnOffset = OffsetToBlock;
+
+		// Find the chunk and offset to the block we need.
+		for (uint32 BlockIndex = 0; BlockIndex < TotalStreamingChunks; ++BlockIndex)
+		{
+			const uint32 SizeOfChunk = StreamingSoundWave->GetSizeOfChunk(BlockIndex);
+
+			if (SizeOfChunk > RemnOffset)
+			{
+				// This is the block we need				
+				if (this->CurrentChunkIndex != BlockIndex)
+				{
+					// Need to seek to another block
+					this->StreamSeekBlockIndex = BlockIndex;
+					this->StreamSeekBlockOffset = RemnOffset;
+				}
+				else
+				{
+					// Seek within this block
+					this->SrcBufferOffset = RemnOffset;
+				}
+				break;
+			}
+
+			RemnOffset -= SizeOfChunk;
+		}
+	}
+}
+
+//-----------------------------------------------------------------------------
+//-----------------------------------------------------------------------------
+bool FBinkAudioInfo::ParseHeader(const uint8* InSrcBufferData, uint32 InSrcBufferDataSize, struct FSoundQualityInfo* QualityInfo)
+{
+	SrcBufferData = InSrcBufferData;
+	SrcBufferDataSize = InSrcBufferDataSize;
+	SrcBufferOffset = 0;
+	CurrentSampleCount = 0;
+
+	check(InSrcBufferDataSize >= sizeof(BinkAudioFileHeader));
+	if (InSrcBufferDataSize < sizeof(BinkAudioFileHeader))
+	{
+		return false;
+	}
+
+	BinkAudioFileHeader* Header = (BinkAudioFileHeader*)InSrcBufferData;
+	if (Header->tag != 'UEBA')
+	{
+		return false;
+	}
+	if (Header->version != 1)
+	{
+		return false;
+	}
+
+	SampleRate = Header->rate;
+	TrueSampleCount = Header->sample_count;
+	NumChannels = Header->channels;
+	MaxCompSpaceNeeded = Header->max_comp_space_needed;
+  
+	uint32 SeekTableSize = Header->seek_table_entry_count * sizeof(uint16);
+	if (sizeof(BinkAudioFileHeader) + SeekTableSize > InSrcBufferDataSize)
+	{
+		return false;
+	}
+
+	// Store the offset to where the audio data begins
+	AudioDataOffset = sizeof(BinkAudioFileHeader) + SeekTableSize;
+
+	// Write out the the header info
+	if (QualityInfo)
+	{
+		QualityInfo->SampleRate = Header->rate;
+		QualityInfo->NumChannels = Header->channels;
+		QualityInfo->SampleDataSize = Header->sample_count * QualityInfo->NumChannels * sizeof(int16);
+		QualityInfo->Duration = (float)Header->sample_count / QualityInfo->SampleRate;
+	}
+
+	return true;
+}
+
+//-----------------------------------------------------------------------------
+//-----------------------------------------------------------------------------
+bool FBinkAudioInfo::CreateDecoder()
+{
+	UEBinkAudioDecodeInterface* BinkInterface = nullptr;
+
+#if WITH_BINK_AUDIO
+	BinkInterface = UnrealBinkAudioDecodeInterface();
+#endif // WITH_BINK_AUDIO
+
+	if (BinkInterface == nullptr)
+	{
+		return false; // only happens if we dont have libs.
+	}
+
+	BinkAudioFileHeader* Header = (BinkAudioFileHeader*)SrcBufferData;
+
+	// Bink is max stereo per stream
+	uint32 StreamCount = (NumChannels + 1) >> 1;	
+
+	// Figure memory for buffers:
+
+	// Deinterlace - buffer space for interleaving multi stream binks in to a standard
+	// interleaved format.
+	uint32 DeinterlaceMemory = 0;
+	if (StreamCount > 1)
+	{
+		DeinterlaceMemory = GetMaxFrameSizeSamples() * sizeof(int16) * 2;
+	}
+
+	// Space to store a decoded block.
+	uint32 OutputReservoirMemory = NumChannels * GetMaxFrameSizeSamples() * sizeof(int16);
+
+	// Space for the decoder state
+	uint32 DecoderMemoryTotal = 0;	
+	uint32 DecoderMemoryPerStream[MAX_BINK_AUDIO_CHANNELS / 2];
+	{
+		uint32 RemnChannels = NumChannels;
+		for (uint32 StreamIndex = 0; StreamIndex < StreamCount; StreamIndex++)
+		{
+			uint32 StreamChannels = RemnChannels;
+			if (StreamChannels > 2)
+			{
+				StreamChannels = 2;
+			}
+			RemnChannels -= StreamChannels;
+			DecoderMemoryPerStream[StreamIndex] = BinkInterface->MemoryFn(SampleRate, StreamChannels);
+			DecoderMemoryTotal += DecoderMemoryPerStream[StreamIndex];
+		}
+	}
+	
+	// Space for the decoder pointers
+	uint32 PtrMemory = Align32(sizeof(void*) * StreamCount);
+	
+	// Space for ourselves + the channel count for each stream.
+	uint32 StructMemory = Align32(sizeof(BinkAudioDecoder) + sizeof(uint8)*StreamCount);
+
+	// Space for decoded seek table
+	uint32 SeekTableMemory = 0;
+	if (Header->seek_table_entry_count)
+	{
+		SeekTableMemory = Align32(sizeof(uint32) * (Header->seek_table_entry_count + 1));
+	}
+
+	uint32 TotalMemory = DecoderMemoryTotal + PtrMemory + StructMemory + OutputReservoirMemory + DeinterlaceMemory + SeekTableMemory;
+
+	//
+	// Allocate and save offsets
+	//
+	RawMemory = (uint8*)FMemory::Malloc(TotalMemory, 16);
+	memset(RawMemory, 0, TotalMemory);
+	
+	Decoder = (BinkAudioDecoder*)RawMemory;
+
+	Decoder->StreamCount = StreamCount;
+	Decoder->OutputReservoirTotalBytes = OutputReservoirMemory;
+	Decoder->SeekTableCount = Header->seek_table_entry_count;
+	Decoder->FramesPerSeekTableEntry = Header->blocks_per_seek_table_entry;
+
+	// See layout discussion in class declaration
+	void** Decoders = Decoder->Decoders();
+	uint8* CurrentMemory = (uint8*)PTR_ADD(Decoders, PtrMemory);
+	uint8* Channels = Decoder->StreamChannels();
+
+	// Init decoders
+	{
+		uint8 RemnChannels = NumChannels;
+		for (uint32 StreamIndex = 0; StreamIndex < StreamCount; StreamIndex++)
+		{
+			uint32 StreamChannels = RemnChannels;
+			if (StreamChannels > 2)
+			{
+				StreamChannels = 2;
+			}
+			RemnChannels -= StreamChannels;
+
+			Channels[StreamIndex] = StreamChannels;
+			Decoders[StreamIndex] = (void**)CurrentMemory;
+			CurrentMemory += DecoderMemoryPerStream[StreamIndex];
+			BinkInterface->OpenFn(Decoders[StreamIndex], SampleRate, StreamChannels, true, true);
+		}
+	}
+
+	Decoder->ToOutputReservoirOffset32 = (uint16)((CurrentMemory - RawMemory) / 32);
+	CurrentMemory += OutputReservoirMemory;
+
+	Decoder->ToDeinterlaceBufferOffset32 = (uint16)((CurrentMemory - RawMemory) / 32);
+	CurrentMemory += DeinterlaceMemory;
+
+	Decoder->ToSeekTableOffset32 = (uint16)((CurrentMemory - RawMemory) / 32);
+	CurrentMemory += SeekTableMemory;
+
+	// Decode the seek table
+	if (Decoder->SeekTableCount)
+	{
+		uint32* SeekTable = Decoder->SeekTable();
+		uint16* EncodedSeekTable = (uint16*)(SrcBufferData + sizeof(BinkAudioFileHeader));
+		uint32 CurrentSeekOffset = 0;
+
+		// the seek table has deltas from last, and we want absolutes
+		for (uint32 i = 0; i < Decoder->SeekTableCount; i++)
+		{
+			SeekTable[i] = CurrentSeekOffset;
+			CurrentSeekOffset += EncodedSeekTable[i];
+		}
+		SeekTable[Decoder->SeekTableCount] = CurrentSeekOffset;
+	}
+
+	SrcBufferOffset = sizeof(BinkAudioFileHeader) + Header->seek_table_entry_count * sizeof(uint16);
+	return true;
+}
+
+//-----------------------------------------------------------------------------
+//-----------------------------------------------------------------------------
+int32 FBinkAudioInfo::GetFrameSize()
+{
+	uint32 BlockSize;
+	if (BinkAudioBlockSize(MaxCompSpaceNeeded, SrcBufferData + SrcBufferOffset, SrcBufferDataSize - SrcBufferOffset, &BlockSize) == false)
+	{
+		// Either malformed data, or not enough data.
+		return 0;
+	}
+
+	return (int32)BlockSize;
+}
+
+//-----------------------------------------------------------------------------
+//-----------------------------------------------------------------------------
+uint32 FBinkAudioInfo::GetMaxFrameSizeSamples() const
+{
+    if (SampleRate >= 44100) 
+    {
+        return 1920;
+    }
+    else if (SampleRate >= 22050) 
+    {
+        return 960;
+    }
+    else 
+    {
+        return 480;
+    }
+}
+
+//-----------------------------------------------------------------------------
+//-----------------------------------------------------------------------------
+FDecodeResult FBinkAudioInfo::Decode(const uint8* CompressedData, const int32 CompressedDataSize, uint8* OutPCMData, const int32 OutputPCMDataSize)
+{
+	UEBinkAudioDecodeInterface* BinkInterface = 0;
+
+#if WITH_BINK_AUDIO
+	BinkInterface = UnrealBinkAudioDecodeInterface();
+#endif // WITH_BINK_AUDIO
+
+	if (BinkInterface == 0)
+	{
+		return FDecodeResult(); // only happens with no libs.
+	}
+
+	uint32 RemnOutputPCMDataSize = OutputPCMDataSize;
+	uint32 RemnCompressedDataSize = CompressedDataSize;
+	const uint8* CompressedDataEnd = CompressedData + CompressedDataSize;
+
+	//
+	// In the event we need to copy to a stack buffer, we alloca() it here so
+	// that it's not inside the loop (for static analysis). We don't touch the memory until we need it
+	// so it's just a couple instructions for the alloca().
+	// (+8 for max block header size)
+	uint8* StackBlockBuffer = (uint8*)alloca(MaxCompSpaceNeeded + BINK_UE_DECODER_END_INPUT_SPACE + 8);
+
+	const uint32 DecodeSize = GetMaxFrameSizeSamples() * sizeof(int16) * NumChannels;
+	while (RemnOutputPCMDataSize)
+	{
+		//
+		// Drain the output reservoir before attempting a decode.
+		//
+		if (Decoder->OutputReservoirValidBytes)
+		{
+			uint32 CopyBytes = Decoder->OutputReservoirValidBytes;
+			if (CopyBytes > RemnOutputPCMDataSize)
+			{
+				CopyBytes = RemnOutputPCMDataSize;
+			}
+
+			FMemory::Memcpy(OutPCMData, Decoder->OutputReservoir(), CopyBytes);
+
+			// We use memmove here because we expect it to be very rare that we don't
+			// consume the entire output reservoir in a call, so it's not worth managing
+			// a cursor. Just move down the remnants, which we expect to be zero.
+			if (Decoder->OutputReservoirValidBytes != CopyBytes)
+			{
+				FMemory::Memmove(Decoder->OutputReservoir(), Decoder->OutputReservoir() + CopyBytes, Decoder->OutputReservoirValidBytes - CopyBytes);
+			}
+			Decoder->OutputReservoirValidBytes -= CopyBytes;
+
+			RemnOutputPCMDataSize -= CopyBytes;
+			OutPCMData += CopyBytes;
+
+			if (RemnOutputPCMDataSize == 0)
+			{
+				// we filled entirely from the output reservoir
+				break;
+			}
+		}
+
+		if (RemnCompressedDataSize == 0)
+		{
+			// This is the normal termination condition
+			break;
+		}
+
+		if (BinkAudioValidateBlock(MaxCompSpaceNeeded, CompressedData, RemnCompressedDataSize) != BINK_AUDIO_BLOCK_VALID)
+		{
+			// The splitting system should ensure that we only ever get complete blocks - so this is bizarre.
+			UE_LOG(LogBinkAudioDecoder, Warning, TEXT("Got weird buffer, validate returned %d"), BinkAudioValidateBlock(MaxCompSpaceNeeded, CompressedData, RemnCompressedDataSize));
+			break;
+		}
+
+		uint8 const* BlockStart =0;
+		uint8 const* BlockEnd =0;
+		uint32 TrimToSampleCount =0;
+		BinkAudioCrackBlock(CompressedData, &BlockStart, &BlockEnd, &TrimToSampleCount);
+		uint8 const* BlockBase = CompressedData;
+
+		//
+		// We need to make sure there's room available for Bink to read past the end
+		// of the buffer (for vector decoding). If there's not, we need to copy to a
+		// temp buffer.
+		//
+		bool HasRoomForDecode = (CompressedDataEnd - BlockEnd) > BINK_UE_DECODER_END_INPUT_SPACE;
+		if (HasRoomForDecode == false)
+		{
+			// This looks weird, but in order for the advancement logic to work,
+			// we need to replicate the entire block including the header.
+			size_t BlockOffset = BlockStart - BlockBase;
+			size_t BlockSize = BlockEnd - BlockBase;
+			if (BlockSize > MaxCompSpaceNeeded + 8) // +8 for max block header size
+			{
+				UE_LOG(LogBinkAudioDecoder, Error, TEXT("BAD! Validated block exceeds header max block size (%d vs %d)"), BlockSize, MaxCompSpaceNeeded);
+				break;
+			}
+
+			FMemory::Memcpy(StackBlockBuffer, BlockBase, BlockSize);
+
+			// this is technically not needed, but just so that any analysis shows that
+			// we've initialized all the memory we touch.
+			FMemory::Memset(StackBlockBuffer + BlockSize, 0, BINK_UE_DECODER_END_INPUT_SPACE);
+
+			BlockBase = StackBlockBuffer;
+			BlockStart = StackBlockBuffer + BlockOffset;
+			BlockEnd = StackBlockBuffer + BlockSize;
+		}
+
+
+		//
+		// If we're a simple single stream and we have enough output space,
+		// just decode directly in to our destination to avoid some
+		// copies.
+		//
+		// We also have to have a "simple" decode - i.e. aligned and no trimming.
+		//
+		if (Decoder->StreamCount == 1 &&
+			DecodeSize <= RemnOutputPCMDataSize &&
+			TrimToSampleCount == ~0U &&
+			(((size_t)OutPCMData)&0xf) == 0 &&
+			Decoder->ConsumeFrameCount == 0)
+		{
+			uint32 DecodedBytes = BinkInterface->DecodeFn(Decoder->Decoders()[0], OutPCMData, RemnOutputPCMDataSize, &BlockStart, BlockEnd);
+			check(DecodedBytes);
+			if (DecodedBytes == 0)
+			{
+				// This means that our block check above succeeded and we still failed - corrupted data!
+				return FDecodeResult();
+			}
+
+			uint32 InputConsumed = (uint32)(BlockStart - BlockBase);
+
+			OutPCMData += DecodedBytes;
+			CompressedData += InputConsumed;
+			RemnCompressedDataSize -= InputConsumed;
+			RemnOutputPCMDataSize -= DecodedBytes;
+			continue;
+		}
+
+		// Otherwise we route through the output reservoir.
+		if (Decoder->StreamCount == 1)
+		{
+			uint32 DecodedBytes = BinkInterface->DecodeFn(Decoder->Decoders()[0], Decoder->OutputReservoir(), Decoder->OutputReservoirTotalBytes, &BlockStart, BlockEnd);
+			check(DecodedBytes);
+			if (DecodedBytes == 0)
+			{
+				// This means that our block check above succeeded and we still failed - corrupted data!
+				return FDecodeResult();
+			}
+
+			uint32 InputConsumed = (uint32)(BlockStart - BlockBase);
+
+			CompressedData += InputConsumed;
+			RemnCompressedDataSize -= InputConsumed;
+
+			Decoder->OutputReservoirValidBytes = DecodedBytes;
+		}
+		else
+		{
+			// multistream requires interlacing the stereo/mono streams
+			int16* DeinterlaceBuffer = Decoder->DeinterlaceBuffer();
+			uint8* CurrentOutputReservoir = Decoder->OutputReservoir();
+
+			uint32 LocalNumChannels = NumChannels;
+			for (uint32 i = 0; i < Decoder->StreamCount; i++)
+			{
+				uint32 StreamChannels = Decoder->StreamChannels()[i];
+
+				uint32 DecodedBytes = BinkInterface->DecodeFn(Decoder->Decoders()[i], (uint8*)DeinterlaceBuffer, GetMaxFrameSizeSamples() * sizeof(int16) * 2, &BlockStart, BlockEnd);
+				check(DecodedBytes);
+
+				if (DecodedBytes == 0)
+				{
+					// This means that our block check above succeeded and we still failed - corrupted data!
+					return FDecodeResult();
+				}
+
+				// deinterleave in to the output reservoir.
+				if (StreamChannels == 1)
+				{
+					int16* Read = DeinterlaceBuffer;
+					int16* Write = (int16*)CurrentOutputReservoir;
+					uint32 Frames = DecodedBytes / sizeof(int16);
+					for (uint32 FrameIndex = 0; FrameIndex < Frames; FrameIndex++)
+					{
+						Write[LocalNumChannels * FrameIndex] = Read[FrameIndex];
+					}
+				}
+				else
+				{
+					// stereo int16 pairs
+					int32* Read = (int32*)DeinterlaceBuffer;
+					int16* Write = (int16*)CurrentOutputReservoir;
+					uint32 Frames = DecodedBytes / sizeof(int32);
+					for (uint32 FrameIndex = 0; FrameIndex < Frames; FrameIndex++)
+					{
+						*(int32*)(Write + LocalNumChannels * FrameIndex) = Read[FrameIndex];
+					}
+				}
+				
+				CurrentOutputReservoir += sizeof(int16) * StreamChannels;
+
+				Decoder->OutputReservoirValidBytes += DecodedBytes;
+			}
+
+			uint32 InputConsumed = (uint32)(BlockStart - BlockBase);			
+			CompressedData += InputConsumed;
+			RemnCompressedDataSize -= InputConsumed;
+		} // end if multi stream
+
+		// Check if we are trimming the tail due to EOF
+		if (TrimToSampleCount != ~0U)
+		{
+			// Ignore the tail samples by just dropping our reservoir size
+			Decoder->OutputReservoirValidBytes = TrimToSampleCount * sizeof(int16) * NumChannels;
+		}
+
+		// Check if we need to eat some frames due to a sample-accurate seek.
+		if (Decoder->ConsumeFrameCount)
+		{
+			const uint32 BytesPerFrame = sizeof(int16) * NumChannels;
+			uint32 ValidFrames = Decoder->OutputReservoirValidBytes / BytesPerFrame;
+			if (Decoder->ConsumeFrameCount < ValidFrames)
+			{
+				FMemory::Memmove(Decoder->OutputReservoir(), Decoder->OutputReservoir() + Decoder->ConsumeFrameCount * BytesPerFrame, (ValidFrames - Decoder->ConsumeFrameCount) * BytesPerFrame);
+				Decoder->OutputReservoirValidBytes -= Decoder->ConsumeFrameCount * BytesPerFrame;
+			}
+			else
+			{
+				Decoder->OutputReservoirValidBytes = 0;
+			}
+			Decoder->ConsumeFrameCount = 0;
+		}
+		// Fall through to the next loop to copy the decoded pcm data out of the reservoir.
+	} // while need output pcm data
+
+	// We get here if we filled the output buffer or not.
+	FDecodeResult Result;
+	Result.NumPcmBytesProduced = OutputPCMDataSize - RemnOutputPCMDataSize;
+	Result.NumAudioFramesProduced = Result.NumPcmBytesProduced / (sizeof(int16) * NumChannels);	
+	Result.NumCompressedBytesConsumed = CompressedDataSize - RemnCompressedDataSize;
+	return Result;
+}
+
+IMPLEMENT_MODULE(FBinkAudioDecoderModule, BinkAudioDecoder)
\ No newline at end of file
diff --git a/Engine/Source/Runtime/BinkAudioDecoder/Public/BinkAudioInfo.h b/Engine/Source/Runtime/BinkAudioDecoder/Public/BinkAudioInfo.h
new file mode 100644
--- /dev/null
+++ b/Engine/Source/Runtime/BinkAudioDecoder/Public/BinkAudioInfo.h
@@ -0,0 +1,34 @@
+// Copyright Epic Games, Inc. All Rights Reserved.
+
+#pragma once
+
+#include "CoreMinimal.h"
+#include "AudioDecompress.h"
+
+/**
+*	Decoder thunk to the Bink Audio libraries. Also manages file parsing
+*	for the cooked data.
+*/
+class  FBinkAudioInfo : public IStreamedCompressedInfo
+{
+public:
+	BINKAUDIODECODER_API FBinkAudioInfo();
+	BINKAUDIODECODER_API virtual ~FBinkAudioInfo();
+
+	//~ Begin IStreamedCompressedInfo Interface
+	virtual bool ParseHeader(const uint8* InSrcBufferData, uint32 InSrcBufferDataSize, FSoundQualityInfo* QualityInfo) override;
+	virtual int32 GetFrameSize() override;
+	virtual uint32 GetMaxFrameSizeSamples() const override;
+	virtual bool CreateDecoder() override;
+	virtual void SeekToTime(const float SeekToTimeSeconds) override;
+	virtual FDecodeResult Decode(const uint8* CompressedData, const int32 CompressedDataSize, uint8* OutPCMData, const int32 OutputPCMDataSize) override;
+	//~ End IStreamedCompressedInfo Interface
+
+protected:
+
+	// copied from header during ParseHeader
+	uint32 MaxCompSpaceNeeded; 
+
+	struct BinkAudioDecoder* Decoder = 0;
+	uint8* RawMemory = 0;
+};
diff --git a/Engine/Source/Runtime/Engine/Public/AudioDecompress.h b/Engine/Source/Runtime/Engine/Public/AudioDecompress.h
--- a/Engine/Source/Runtime/Engine/Public/AudioDecompress.h
+++ b/Engine/Source/Runtime/Engine/Public/AudioDecompress.h
@@ -190,7 +190,7 @@ struct FDecodeResult
 class IStreamedCompressedInfo : public ICompressedAudioInfo
 {
 public:
-	IStreamedCompressedInfo();
+	ENGINE_API IStreamedCompressedInfo();
 	virtual ~IStreamedCompressedInfo() {}
 
 	//~ Begin ICompressedInfo Interface
@@ -203,8 +203,8 @@ public:
 	virtual bool UsesVorbisChannelOrdering() const override { return false; }
 	virtual int GetStreamBufferSize() const override { return  MONO_PCM_BUFFER_SIZE; }
 	virtual bool SupportsStreaming() const override { return true; }
-	virtual bool StreamCompressedInfoInternal(const FSoundWaveProxyPtr& InWaveProxy, FSoundQualityInfo* QualityInfo) override;
-	virtual bool StreamCompressedData(uint8* Destination, bool bLooping, uint32 BufferSize, int32& OutNumBytesStreamed) override;
+	ENGINE_API virtual bool StreamCompressedInfoInternal(const FSoundWaveProxyPtr& InWaveProxy, FSoundQualityInfo* QualityInfo) override;
+	ENGINE_API virtual bool StreamCompressedData(uint8* Destination, bool bLooping, uint32 BufferSize, int32& OutNumBytesStreamed) override;
 	virtual int32 GetCurrentChunkIndex() const override { return CurrentChunkIndex; }
 	virtual int32 GetCurrentChunkOffset() const override { return SrcBufferOffset; }
 	virtual bool IsStreamedCompressedInfo() const override { return true; }
