Bink Audio support for >48khz. Well... support in that they play, however Bink Audio (like most perceptual codecs) don't bother to send such high frequencies, so I add a warning to resample down prior to encoding. It's just a waste of performance.

This looks like a large change but the relevant change is binka_ue_file_header.h. Then rebuilt libs, and a bunch of source changes that never made it over in the initial checkin.

diff --git a/Engine/Source/Developer/AudioFormatBink/Private/AudioFormatBink.cpp b/Engine/Source/Developer/AudioFormatBink/Private/AudioFormatBink.cpp
--- a/Engine/Source/Developer/AudioFormatBink/Private/AudioFormatBink.cpp
+++ b/Engine/Source/Developer/AudioFormatBink/Private/AudioFormatBink.cpp
@@ -17,7 +17,7 @@ class FAudioFormatBink : public IAudioFormat
 	enum
 	{
 		/** Version for Bink Audio format, this becomes part of the DDC key. */
-		UE_AUDIO_BINK_VER = 1,
+		UE_AUDIO_BINK_VER = 2,
 	};
 
 public:
diff --git a/Engine/Source/Runtime/BinkAudioDecoder/SDK/BinkAudio/Include/binka_ue_file_header.h b/Engine/Source/Runtime/BinkAudioDecoder/SDK/BinkAudio/Include/binka_ue_file_header.h
--- a/Engine/Source/Runtime/BinkAudioDecoder/SDK/BinkAudio/Include/binka_ue_file_header.h
+++ b/Engine/Source/Runtime/BinkAudioDecoder/SDK/BinkAudio/Include/binka_ue_file_header.h
@@ -12,7 +12,8 @@ struct BinkAudioFileHeader
     uint32 tag; // UEBA
     uint8 version;
     uint8 channels;
-    uint16 rate;
+    uint16 PADDING; // unused atm.
+    uint32 rate;
     uint32 sample_count;
 
     // the max for a bink audio block - includes all streams, but does _not_ include
diff --git a/Engine/Source/Runtime/BinkAudioDecoder/SDK/BinkAudio/Src/binka_ue_decode.cpp b/Engine/Source/Runtime/BinkAudioDecoder/SDK/BinkAudio/Src/binka_ue_decode.cpp
--- a/Engine/Source/Runtime/BinkAudioDecoder/SDK/BinkAudio/Src/binka_ue_decode.cpp
+++ b/Engine/Source/Runtime/BinkAudioDecoder/SDK/BinkAudio/Src/binka_ue_decode.cpp
@@ -4,13 +4,13 @@
 #include "binka_ue_decode.h"
 
 // Amount of memory to pass to DecompressOpen
-uint32_t UEBinkAudioDecompressMemory(uint32_t rate, uint32_t chans)
+static uint32_t UEBinkAudioDecompressMemory(uint32_t rate, uint32_t chans)
 {
     return BinkAudioDecompressMemory(rate, chans, 0);
 }
 
 // open and initialize a decompression stream
-uint32_t UEBinkAudioDecompressOpen(void * BinkAudioDecoderMemory, uint32_t rate, uint32_t chans, bool interleave_output, bool is_bink_audio_2)
+static uint32_t UEBinkAudioDecompressOpen(void * BinkAudioDecoderMemory, uint32_t rate, uint32_t chans, bool interleave_output, bool is_bink_audio_2)
 {
     uint32_t flags = 0;
     if (is_bink_audio_2)
@@ -21,7 +21,7 @@ uint32_t UEBinkAudioDecompressOpen(void * BinkAudioDecoderMemory, uint32_t rate,
 }
 
 //do the decompression - supports linear or ringbuffers (outptr==outstart on non-ring), will clamp, if no room
-uint32_t UEBinkAudioDecompress(void * BinkAudioDecoderMemory, uint8_t* OutputBuffer, uint32_t OutputBufferLen, uint8_t const** InputBuffer, uint8_t const* InputBufferEnd)
+static uint32_t UEBinkAudioDecompress(void * BinkAudioDecoderMemory, uint8_t* OutputBuffer, uint32_t OutputBufferLen, uint8_t const** InputBuffer, uint8_t const* InputBufferEnd)
 {
     BINKAC_OUT_RINGBUF out;
     BINKAC_IN in;
@@ -45,7 +45,7 @@ uint32_t UEBinkAudioDecompress(void * BinkAudioDecoderMemory, uint8_t* OutputBuf
 }
 
 // resets the start flag to prevent blending in the last decoded frame.
-void UEBinkAudioDecompressResetStartFrame(void * BinkAudioDecoderMemory)
+static void UEBinkAudioDecompressResetStartFrame(void * BinkAudioDecoderMemory)
 {
     BinkAudioDecompressResetStartFrame(BinkAudioDecoderMemory);
 }
diff --git a/Engine/Source/Runtime/BinkAudioDecoder/SDK/BinkAudio/Src/binka_ue_encode.cpp b/Engine/Source/Runtime/BinkAudioDecoder/SDK/BinkAudio/Src/binka_ue_encode.cpp
--- a/Engine/Source/Runtime/BinkAudioDecoder/SDK/BinkAudio/Src/binka_ue_encode.cpp
+++ b/Engine/Source/Runtime/BinkAudioDecoder/SDK/BinkAudio/Src/binka_ue_encode.cpp
@@ -354,7 +354,7 @@ uint8_t UECompressBinkAudio(
     BinkAudioFileHeader Header;
     Header.tag = 'UEBA';
     Header.channels = (U8)WavChannels;
-    Header.rate = (U16)WavRate;
+    Header.rate = WavRate;
     Header.sample_count = SamplesPerChannel;
     Header.max_comp_space_needed = (U16)MaxBlockSize;
     Header.flags = 1;
diff --git a/Engine/Source/Runtime/BinkAudioDecoder/SDK/BinkAudio/Src/binkace.h b/Engine/Source/Runtime/BinkAudioDecoder/SDK/BinkAudio/Src/binkace.h
--- a/Engine/Source/Runtime/BinkAudioDecoder/SDK/BinkAudio/Src/binkace.h
+++ b/Engine/Source/Runtime/BinkAudioDecoder/SDK/BinkAudio/Src/binkace.h
@@ -9,6 +9,16 @@
 
 #define AUDIOFLOAT F32
 
+#ifdef WRAP_PUBLICS
+#define rfmerge3(name,add) name##add
+#define rfmerge2(name,add) rfmerge3(name,add)
+#define rfmerge(name)      rfmerge2(name,WRAP_PUBLICS)
+#define BinkAudioCompressOpen                          rfmerge(BinkAudioCompressOpen)
+#define BinkAudioCompressLock                          rfmerge(BinkAudioCompressLock)
+#define BinkAudioCompressUnlock                        rfmerge(BinkAudioCompressUnlock)
+#define BinkAudioCompressClose                         rfmerge(BinkAudioCompressClose)
+#endif
+
 //===========================================
 //  encoding API
 //===========================================
diff --git a/Engine/Source/Runtime/BinkAudioDecoder/SDK/BinkAudio/Src/cpu.h b/Engine/Source/Runtime/BinkAudioDecoder/SDK/BinkAudio/Src/cpu.h
--- a/Engine/Source/Runtime/BinkAudioDecoder/SDK/BinkAudio/Src/cpu.h
+++ b/Engine/Source/Runtime/BinkAudioDecoder/SDK/BinkAudio/Src/cpu.h
@@ -54,6 +54,17 @@
 
 #elif defined( __RADX86__ )
 
+#ifdef WRAP_PUBLICS
+#define rfmerge3(name,add) name##add
+#define rfmerge2(name,add) rfmerge3(name,add)
+#define rfmerge(name)      rfmerge2(name,WRAP_PUBLICS)
+#define cpu_features_to_use             rfmerge(cpu_features_to_use)
+#define cpu_features_avail              rfmerge(cpu_features_avail)
+#define CPU_check                       rfmerge(CPU_check)
+#define CPU_clear                       rfmerge(CPU_clear)
+#endif
+
+
 RADDECLAREDATA S32 cpu_features_to_use;
 RADDECLAREDATA S32 cpu_features_avail;
 
diff --git a/Engine/Source/Runtime/BinkAudioDecoder/SDK/BinkAudio/Src/popmal.c b/Engine/Source/Runtime/BinkAudioDecoder/SDK/BinkAudio/Src/popmal.c
--- a/Engine/Source/Runtime/BinkAudioDecoder/SDK/BinkAudio/Src/popmal.c
+++ b/Engine/Source/Runtime/BinkAudioDecoder/SDK/BinkAudio/Src/popmal.c
@@ -44,7 +44,6 @@ RADDEFFUNC void RADLINK pushmalloc(void* base, void * ptr,U64 amt)
     RR_BREAK();
   }
 
-#ifndef SEPARATEMALLOCS
   {
     U64 last,next;
 
@@ -55,7 +54,6 @@ RADDEFFUNC void RADLINK pushmalloc(void* base, void * ptr,U64 amt)
     // make sure the up 32 separate mallocs use distinct sets
     amt+=(((32+last-next)&31)*32);
   }
-#endif
 
   p->pushtot+=amt;
   p->pushamt[p->pushcur]=amt;
@@ -76,90 +74,7 @@ RADDEFFUNC U64 RADLINK popmalloctotal( void* base )
   return( p->pushtot );
 }
 
-#ifdef SEPARATEMALLOCS
-
-#ifdef __RADFINAL__
-#error "You have separate mallocs turned on!"
-#endif
-
-#ifdef RADUSETM3
-RADDEFFUNC void * RADLINK popmalloci(void* base, U64 amt, char const * info, U32 line)
-#else
-RADDEFFUNC void * RADLINK popmalloc(void* base, U64 amt)
-#endif
-{
-  void * ptr;
-  UINTa * table;
-  U32 i;
-  PM * p = (PM *)base;
-  U32 extra;
-
-  extra = sizeof(UINTa)+sizeof(void*)+sizeof(void*)+sizeof(UINTa);
-  if ( p )
-    extra += ( p->pushcur * sizeof(void*) );
-  extra = (extra+15) & ~15;
-
-  ptr=radmalloc(amt+extra);
-  table = (UINTa *)ptr;
-  table[ 0 ] = 1;
-  table[ 1 ] = (UINTa)table;
-
-  ptr = ( (U8 *) ptr ) + extra;
-  ((UINTa*)ptr)[-1] = 0x97537531;
-  ((UINTa*)ptr)[-2] = (UINTa)table;
-
-  if ( p )
-  {
-    table[ 0 ] = 1 + p->pushcur;
-    p->pushtot=0;
-    for(i=0;i<p->pushcur;i++)
-    {
-      void * np;
-#ifdef RADUSETM3
-      np=radmalloci(p->pushamt[i]&~1,info,line);
-#else
-      np=radmalloc(p->pushamt[i]&~1);
-#endif
-
-      if ( p->pushamt[i]&1 )
-        (*(void**)(((U8*)ptr)+((UINTa)p->pushptr[i])))=np;
-      else
-        (*(p->pushptr[i]))=np;
-
-      table[ i + 1 ] = (UINTa)np;
-    }
-    table[ i + 1 ] = (UINTa)table;
-    p->pushcur=0;
-  }
-  return(ptr);
-}
-
-RADDEFFUNC void RADLINK popfree(void * ptr)
-{
-  UINTa * table;
-  UINTa i, n;
-  if ( ((UINTa*)ptr)[-1] != 0x97537531 )
-  {
-     RR_BREAK();
-     return;
-  }
-  ((UINTa*)ptr)[-1] = 0;
-  table = ((UINTa**)ptr)[-2];
-  n = table[0];
-  table[0]=0;
-  for(i=0;i<n;i++)
-  {
-    radfree( (void*)table[i+1] );
-  }
-}
-
-#else
-
-#ifdef RADUSETM3
-RADDEFFUNC void * RADLINK popmalloci(void* base, U64 amt, char const * info, U32 line)
-#else
 RADDEFFUNC void * RADLINK popmalloc(void* base, U64 amt, void* (*allocator)(UINTa bytes))
-#endif
 {
   PM * p = (PM *)base;
   void * ptr;
@@ -171,11 +86,7 @@ RADDEFFUNC void * RADLINK popmalloc(void* base, U64 amt, void* (*allocator)(UINT
     return allocator( (UINTa)amt );
   }
 
-#ifdef RADUSETM3
-  ptr=radmalloci(p->pushtot+amt,info,line);
-#else
   ptr=allocator((UINTa)(p->pushtot+amt));
-#endif
   
   p->pushtot=0;
   if (ptr) {
@@ -197,6 +108,3 @@ RADDEFFUNC void * RADLINK popmalloc(void* base, U64 amt, void* (*allocator)(UINT
   p->pushcur=0;
   return(ptr);
 }
-
-#endif
-
diff --git a/Engine/Source/Runtime/BinkAudioDecoder/SDK/BinkAudio/Src/popmal.h b/Engine/Source/Runtime/BinkAudioDecoder/SDK/BinkAudio/Src/popmal.h
--- a/Engine/Source/Runtime/BinkAudioDecoder/SDK/BinkAudio/Src/popmal.h
+++ b/Engine/Source/Runtime/BinkAudioDecoder/SDK/BinkAudio/Src/popmal.h
@@ -10,6 +10,19 @@
 
 #define PushMallocBytesForXPtrs( X ) ( ( X * sizeof(void*) ) + ( X * sizeof(U64) ) + 64 )
 
+
+#ifdef WRAP_PUBLICS
+#define rfmerge3(name,add) name##add
+#define rfmerge2(name,add) rfmerge3(name,add)
+#define rfmerge(name)      rfmerge2(name,WRAP_PUBLICS)
+#define pushmallocinit                    rfmerge(pushmallocinit)
+#define pushmalloc                        rfmerge(pushmalloc)
+#define pushmalloco                       rfmerge(pushmalloco)
+#define popmalloctotal                    rfmerge(popmalloctotal)
+#define popmalloc                         rfmerge(popmalloc)
+#endif
+
+
 RADDEFFUNC void RADLINK pushmallocinit(void * base,U32 num_ptrs);
 RADDEFFUNC void RADLINK pushmalloc( void * base, void  * ptr, U64 amt );
 
@@ -18,19 +31,8 @@ RADDEFFUNC void RADLINK pushmalloco(void* base, void * ptr,U64 amt);
 
 RADDEFFUNC U64 RADLINK popmalloctotal( void * base );
 
-#ifdef RADUSETM3
-  #define popmalloc( base,b ) popmalloci( base, b, __FILE__, __LINE__ )
-  RADDEFFUNC void  * RADLINK popmalloci( void * base, U64 amt, char const * info, U32 line );
-#else
-  RADDEFFUNC void  * RADLINK popmalloc( void * base, U64 amt, void* (*allocator)(UINTa bytes));
-#endif
-
-//#define SEPARATEMALLOCS
+RADDEFFUNC void  * RADLINK popmalloc( void * base, U64 amt, void* (*allocator)(UINTa bytes));
 
-#ifdef SEPARATEMALLOCS
-RADDEFFUNC void RADLINK popfree(void  * ptr);
-#else
 #define popfree(ptr, memfree) memfree(ptr)
-#endif
 
-#endif
+#endif //  __POPMALH__
diff --git a/Engine/Source/Runtime/BinkAudioDecoder/SDK/BinkAudio/Src/varbits.h b/Engine/Source/Runtime/BinkAudioDecoder/SDK/BinkAudio/Src/varbits.h
--- a/Engine/Source/Runtime/BinkAudioDecoder/SDK/BinkAudio/Src/varbits.h
+++ b/Engine/Source/Runtime/BinkAudioDecoder/SDK/BinkAudio/Src/varbits.h
@@ -6,10 +6,15 @@
 #include "rrCore.h"
 #endif
 
-#ifndef __RADMATHH__
-//#include "radmath.h"
+#ifdef WRAP_PUBLICS
+#define rfmerge3(name,add) name##add
+#define rfmerge2(name,add) rfmerge3(name,add)
+#define rfmerge(name)      rfmerge2(name,WRAP_PUBLICS)
+#define VarBitsCopy                          rfmerge(VarBitsCopy)
+#define VarBitsLens                          rfmerge(VarBitsLens)
 #endif
 
+
 // variable bit macros
 #if defined(__RADX86__)
 RADDECLAREDATA const RAD_ALIGN(U32,VarBitsLens[33],32);
diff --git a/Engine/Source/Runtime/Engine/Private/AudioDerivedData.cpp b/Engine/Source/Runtime/Engine/Private/AudioDerivedData.cpp
--- a/Engine/Source/Runtime/Engine/Private/AudioDerivedData.cpp
+++ b/Engine/Source/Runtime/Engine/Private/AudioDerivedData.cpp
@@ -1168,6 +1168,14 @@ static void CookSimpleWave(USoundWave* SoundWave, FName FormatName, const IAudio
 		QualityInfo.bStreaming = SoundWave->IsStreaming(CompressionOverrides ? *CompressionOverrides : FPlatformAudioCookOverrides());
 		QualityInfo.DebugName = SoundWave->GetFullName();
 
+		static const FName NAME_BINKA(TEXT("BINKA"));
+		if (WaveSampleRate > 48000 &&
+			FormatName == NAME_BINKA)
+		{
+			// We have to do this here because we don't know the name of the wave inside the codec.
+			UE_LOG(LogAudioDerivedData, Warning, TEXT("[%s] High sample rate wave (%d) with Bink Audio - perf waste - high frequencies are discarded by Bink Audio (like most perceptual codecs)."), *SoundWave->GetFullName(), WaveSampleRate);
+		}
+
 		// Cook the data.
 		if(Format.Cook(FormatName, Input, QualityInfo, Output))
 		{
@@ -1380,6 +1388,15 @@ static void CookSurroundWave( USoundWave* SoundWave, FName FormatName, const IAu
 			QualityInfo.SampleDataSize = SampleDataSize;
 			QualityInfo.bStreaming = SoundWave->IsStreaming(CompressionOverrides ? *CompressionOverrides : FPlatformAudioCookOverrides());
 			QualityInfo.DebugName = SoundWave->GetFullName();
+
+			static const FName NAME_BINKA(TEXT("BINKA"));
+			if (WaveSampleRate > 48000 &&
+				FormatName == NAME_BINKA)
+			{
+				// We have to do this here because we don't know the name of the wave inside the codec.
+				UE_LOG(LogAudioDerivedData, Warning, TEXT("[%s] High sample rate wave (%d) with Bink Audio - perf waste - high frequencies are discarded by Bink Audio (like most perceptual codecs)."), *SoundWave->GetFullName(), WaveSampleRate);
+			}
+
 			//@todo tighten up the checking for empty results here
 			if(Format.CookSurround(FormatName, SourceBuffers, QualityInfo, Output))
 			{
