Switching to a fat binary for OSX bink audio static libs, and adding osx arm 64 for M1

diff --git a/Engine/Source/Runtime/BinkAudioDecoder/Module/BinkAudioDecoder.Build.cs b/Engine/Source/Runtime/BinkAudioDecoder/Module/BinkAudioDecoder.Build.cs
--- a/Engine/Source/Runtime/BinkAudioDecoder/Module/BinkAudioDecoder.Build.cs
+++ b/Engine/Source/Runtime/BinkAudioDecoder/Module/BinkAudioDecoder.Build.cs
@@ -20,7 +20,7 @@ public class BinkAudioDecoder : ModuleRules
         }
         if (Target.Platform == UnrealTargetPlatform.Mac)
         {
-            return Path.Combine(ModuleDirectory, "..", "SDK", "BinkAudio", "Lib", "libbinka_ue_decode_osx64_static.a");
+            return Path.Combine(ModuleDirectory, "..", "SDK", "BinkAudio", "Lib", "libbinka_ue_decode_osxfat_static.a");
         }
         if (Target.Platform == UnrealTargetPlatform.IOS)
         {
