BinkAudio Error flag.

diff --git a/Engine/Source/Runtime/BinkAudioDecoder/Module/Private/BinkAudioInfo.cpp b/Engine/Source/Runtime/BinkAudioDecoder/Module/Private/BinkAudioInfo.cpp
--- a/Engine/Source/Runtime/BinkAudioDecoder/Module/Private/BinkAudioInfo.cpp
+++ b/Engine/Source/Runtime/BinkAudioDecoder/Module/Private/BinkAudioInfo.cpp
@@ -376,6 +376,9 @@ int32 FBinkAudioInfo::GetFrameSize()
 	uint32 BlockSize;
 	if (BinkAudioBlockSize(MaxCompSpaceNeeded, SrcBufferData + SrcBufferOffset, SrcBufferDataSize - SrcBufferOffset, &BlockSize) == false)
 	{
+		// Flag this as error so that the owning logic can clean up this decode.
+		bErrorStateLatch = true;
+		
 		// Either malformed data, or not enough data.
 		return 0;
 	}
@@ -496,6 +499,7 @@ FDecodeResult FBinkAudioInfo::Decode(const uint8* CompressedData, const int32 Co
 			if (BlockSize > MaxCompSpaceNeeded + 8) // +8 for max block header size
 			{
 				UE_LOG(LogBinkAudioDecoder, Error, TEXT("BAD! Validated block exceeds header max block size (%d vs %d)"), BlockSize, MaxCompSpaceNeeded);
+				bErrorStateLatch = true;
 				break;
 			}
 
@@ -528,6 +532,8 @@ FDecodeResult FBinkAudioInfo::Decode(const uint8* CompressedData, const int32 Co
 			check(DecodedBytes);
 			if (DecodedBytes == 0)
 			{
+				bErrorStateLatch = true;
+
 				// This means that our block check above succeeded and we still failed - corrupted data!
 				return FDecodeResult();
 			}
@@ -548,6 +554,8 @@ FDecodeResult FBinkAudioInfo::Decode(const uint8* CompressedData, const int32 Co
 			check(DecodedBytes);
 			if (DecodedBytes == 0)
 			{
+				bErrorStateLatch = true;
+
 				// This means that our block check above succeeded and we still failed - corrupted data!
 				return FDecodeResult();
 			}
@@ -646,4 +654,9 @@ FDecodeResult FBinkAudioInfo::Decode(const uint8* CompressedData, const int32 Co
 	return Result;
 }
 
-IMPLEMENT_MODULE(FBinkAudioDecoderModule, BinkAudioDecoder)
\ No newline at end of file
+bool FBinkAudioInfo::HasError() const
+{
+	return bErrorStateLatch;
+}
+
+IMPLEMENT_MODULE(FBinkAudioDecoderModule, BinkAudioDecoder)
diff --git a/Engine/Source/Runtime/BinkAudioDecoder/Module/Public/BinkAudioInfo.h b/Engine/Source/Runtime/BinkAudioDecoder/Module/Public/BinkAudioInfo.h
--- a/Engine/Source/Runtime/BinkAudioDecoder/Module/Public/BinkAudioInfo.h
+++ b/Engine/Source/Runtime/BinkAudioDecoder/Module/Public/BinkAudioInfo.h
@@ -22,6 +22,7 @@ public:
 	virtual bool CreateDecoder() override;
 	virtual void SeekToTime(const float SeekToTimeSeconds) override;
 	virtual FDecodeResult Decode(const uint8* CompressedData, const int32 CompressedDataSize, uint8* OutPCMData, const int32 OutputPCMDataSize) override;
+	virtual bool HasError() const override;
 	//~ End IStreamedCompressedInfo Interface
 
 protected:
@@ -32,4 +33,5 @@ protected:
 
 	struct BinkAudioDecoder* Decoder = 0;
 	uint8* RawMemory = 0;
+	bool bErrorStateLatch = false;
 };
