Rebuild bink audio SDK libraries as universal binaries for macOS

- on macOS Catalina 10.15.7
- with Xcode 12.4
- as universal binaries with arm64 and x86_64 slices

- Add BuildForMac.command script
- Add CMakeLists.txt

#rnx

diff --git a/Engine/Source/Developer/AudioFormatBink/AudioFormatBink.Build.cs b/Engine/Source/Developer/AudioFormatBink/AudioFormatBink.Build.cs
--- a/Engine/Source/Developer/AudioFormatBink/AudioFormatBink.Build.cs
+++ b/Engine/Source/Developer/AudioFormatBink/AudioFormatBink.Build.cs
@@ -26,7 +26,7 @@ public class AudioFormatBink : ModuleRules
 		}
 		if (Target.Platform == UnrealTargetPlatform.Mac)
 		{
-			PublicAdditionalLibraries.Add(Path.Combine(ModuleDirectory, "..", "..", "Runtime", "BinkAudioDecoder", "SDK", "BinkAudio", "Lib", "libbinka_ue_encode_osx64_static.a"));
+			PublicAdditionalLibraries.Add(Path.Combine(ModuleDirectory, "..", "..", "Runtime", "BinkAudioDecoder", "SDK", "BinkAudio", "Lib", "libbinka_ue_encode_osx_static.a"));
 		}
 	}
 }
diff --git a/Engine/Source/Runtime/BinkAudioDecoder/Module/BinkAudioDecoder.Build.cs b/Engine/Source/Runtime/BinkAudioDecoder/Module/BinkAudioDecoder.Build.cs
--- a/Engine/Source/Runtime/BinkAudioDecoder/Module/BinkAudioDecoder.Build.cs
+++ b/Engine/Source/Runtime/BinkAudioDecoder/Module/BinkAudioDecoder.Build.cs
@@ -20,7 +20,7 @@ public class BinkAudioDecoder : ModuleRules
         }
         if (Target.Platform == UnrealTargetPlatform.Mac)
         {
-            return Path.Combine(ModuleDirectory, "..", "SDK", "BinkAudio", "Lib", "libbinka_ue_decode_osxfat_static.a");
+            return Path.Combine(ModuleDirectory, "..", "SDK", "BinkAudio", "Lib", "libbinka_ue_decode_osx_static.a");
         }
         if (Target.Platform == UnrealTargetPlatform.IOS)
         {
diff --git a/Engine/Source/Runtime/BinkAudioDecoder/SDK/BinkAudio/BuildForMac.command b/Engine/Source/Runtime/BinkAudioDecoder/SDK/BinkAudio/BuildForMac.command
new file mode 100755
--- /dev/null
+++ b/Engine/Source/Runtime/BinkAudioDecoder/SDK/BinkAudio/BuildForMac.command
@@ -0,0 +1,45 @@
+#!/bin/bash
+
+set -e
+UE_THIRD_PARTY_LOCATION=`cd $(pwd)/..; pwd`
+
+UE_MODULE_LOCATION=`pwd`
+
+SOURCE_LOCATION="$UE_MODULE_LOCATION"
+BUILD_LOCATION="$UE_MODULE_LOCATION/Intermediate"
+
+# Specify all of the include/bin/lib directory variables so that CMake can
+# compute relative paths correctly for the imported targets.
+INSTALL_LOCATION="$UE_MODULE_LOCATION/Lib/"
+
+rm -rf $BUILD_LOCATION
+mkdir $BUILD_LOCATION
+pushd $BUILD_LOCATION
+
+CMAKE_ARGS=(
+    -DCMAKE_INSTALL_PREFIX="$INSTALL_LOCATION"
+    -DCMAKE_OSX_DEPLOYMENT_TARGET="10.9"
+    -DCMAKE_OSX_ARCHITECTURES="arm64;x86_64"
+)
+
+echo Configuring build for BinkAudio...
+cmake -G "Xcode" $SOURCE_LOCATION "${CMAKE_ARGS[@]}"
+
+echo Building BinkAudio for Debug...
+cmake --build . --config Debug
+
+echo Installing BinkAudio for Debug...
+cmake --install . --config Debug
+
+echo Building BinkAudio for Release...
+cmake --build . --config Release
+
+echo Installing BinkAudio for Release...
+cmake --install . --config Release
+
+popd
+
+cp "$BUILD_LOCATION/Release/liblibbinka_ue_encode.a" "$INSTALL_LOCATION/libbinka_ue_encode_osx_static.a"
+cp "$BUILD_LOCATION/Release/liblibbinka_ue_decode.a" "$INSTALL_LOCATION/libbinka_ue_decode_osx_static.a"
+
+echo Done.
diff --git a/Engine/Source/Runtime/BinkAudioDecoder/SDK/BinkAudio/CMakeLists.txt b/Engine/Source/Runtime/BinkAudioDecoder/SDK/BinkAudio/CMakeLists.txt
new file mode 100644
--- /dev/null
+++ b/Engine/Source/Runtime/BinkAudioDecoder/SDK/BinkAudio/CMakeLists.txt
@@ -0,0 +1,23 @@
+cmake_minimum_required(VERSION 3.00)
+PROJECT(BinkAudioDecoder)
+
+SET(ENCODER_SRCS
+	Src/popmal.c
+    	Src/binkace.c
+    	Src/binka_ue_encode.cpp
+    	Src/radfft.cpp
+    	Src/varbits.c)
+
+SET(DECODER_SRCS
+    Src/binkacd.cpp
+    Src/binka_ue_decode.cpp
+    Src/radfft.cpp
+    Src/x86_cpu.c)
+
+include_directories(Include/)
+
+ADD_DEFINITIONS(-D__RADINSTATICLIB__=1)
+ADD_LIBRARY(libbinka_ue_encode ${ENCODER_SRCS})
+
+ADD_DEFINITIONS(-DWRAP_PUBLICS=UEBA)
+ADD_LIBRARY(libbinka_ue_decode ${DECODER_SRCS})
diff --git a/Engine/Source/Runtime/BinkAudioDecoder/SDK/BinkAudio/Src/varbits.h b/Engine/Source/Runtime/BinkAudioDecoder/SDK/BinkAudio/Src/varbits.h
--- a/Engine/Source/Runtime/BinkAudioDecoder/SDK/BinkAudio/Src/varbits.h
+++ b/Engine/Source/Runtime/BinkAudioDecoder/SDK/BinkAudio/Src/varbits.h
@@ -16,9 +16,7 @@
 
 
 // variable bit macros
-#if defined(__RADX86__)
 RADDECLAREDATA const RAD_ALIGN(U32,VarBitsLens[33],32);
-#endif
 
 /*************
 
diff --git a/Engine/Source/Runtime/BinkAudioDecoder/SDK/BinkAudio/Src/x86_cpu.c b/Engine/Source/Runtime/BinkAudioDecoder/SDK/BinkAudio/Src/x86_cpu.c
--- a/Engine/Source/Runtime/BinkAudioDecoder/SDK/BinkAudio/Src/x86_cpu.c
+++ b/Engine/Source/Runtime/BinkAudioDecoder/SDK/BinkAudio/Src/x86_cpu.c
@@ -1,6 +1,8 @@
 // Copyright Epic Games, Inc. All Rights Reserved.
 #include "cpu.h"
 
+#if defined(__RADX86__)
+
 RADDEFINEDATA S32 cpu_features_to_use;
 RADDEFINEDATA S32 cpu_features_avail = 0;
 
@@ -106,3 +108,5 @@ RADDEFFUNC void CPU_clear( void )
   _mm_empty();
 #endif
 }
+
+#endif // defined(__RADX86__)
