use bink audio on windows-based platform extensions.

#rnx

diff --git a/Engine/Source/Runtime/BinkAudioDecoder/Module/BinkAudioDecoder.Build.cs b/Engine/Source/Runtime/BinkAudioDecoder/Module/BinkAudioDecoder.Build.cs
--- a/Engine/Source/Runtime/BinkAudioDecoder/Module/BinkAudioDecoder.Build.cs
+++ b/Engine/Source/Runtime/BinkAudioDecoder/Module/BinkAudioDecoder.Build.cs
@@ -10,7 +10,7 @@ public class BinkAudioDecoder : ModuleRules
     // virtual so that NDA platforms can hide secrets
     protected virtual string GetLibrary(ReadOnlyTargetRules Target)
     {
-        if (Target.Platform == UnrealTargetPlatform.Win64)
+        if (Target.Platform.IsInGroup(UnrealPlatformGroup.Windows))
         {
             return Path.Combine(ModuleDirectory, "..", "SDK", "BinkAudio", "Lib", "binka_ue_decode_win64_static.lib");
         }
