update sample count when seeking in Bink Audio Decoder #preflight 62b04a8e6965d8cf8de29b60 #jira UE-157059

diff --git a/Engine/Source/Runtime/BinkAudioDecoder/Module/Private/BinkAudioInfo.cpp b/Engine/Source/Runtime/BinkAudioDecoder/Module/Private/BinkAudioInfo.cpp
--- a/Engine/Source/Runtime/BinkAudioDecoder/Module/Private/BinkAudioInfo.cpp
+++ b/Engine/Source/Runtime/BinkAudioDecoder/Module/Private/BinkAudioInfo.cpp
@@ -123,6 +123,7 @@ void FBinkAudioInfo::SeekToTime(const float SeekTimeSeconds)
 	{
 		SeekTimeSamples = this->TrueSampleCount - 1;
 	}
+	this->CurrentSampleCount = SeekTimeSamples;
 
 	uint32 SamplesPerBlock = SamplesInFrame * Decoder->FramesPerSeekTableEntry;
 	uint32 SeekTableIndex = SeekTimeSamples / SamplesPerBlock;
@@ -659,4 +660,4 @@ bool FBinkAudioInfo::HasError() const
 	return bErrorStateLatch;
 }
 
-IMPLEMENT_MODULE(FBinkAudioDecoderModule, BinkAudioDecoder)
+IMPLEMENT_MODULE(FBinkAudioDecoderModule, BinkAudioDecoder)
\ No newline at end of file
