bink audio arm64 encoding consistency, and picking up a handful of minor tweaks/fixes. This also adds scripts for building the various apple libs that don't require cmake.

Currently we still rely on a handful of libm functions to be consistent across compilers/plaforms, however current consistency tests pass.

Due to the (very) narrow range we call log, a custom approximation is added to ensure consistency.

Additionally, we reset the floating point state on entry in case fp:fast has messed with it on our thread.

Note that the decoder libs are not currently updated as none of the fixes are worth the testing load due to lack of active bug reports.

diff --git a/Engine/Source/Runtime/BinkAudioDecoder/SDK/BinkAudio/BuildDecoderIos.command b/Engine/Source/Runtime/BinkAudioDecoder/SDK/BinkAudio/BuildDecoderIos.command
new file mode 100755
--- /dev/null
+++ b/Engine/Source/Runtime/BinkAudioDecoder/SDK/BinkAudio/BuildDecoderIos.command
@@ -0,0 +1,113 @@
+#!/bin/bash
+
+set -e
+
+UE_MODULE_LOCATION=`pwd`
+
+SOURCE_LOCATION="$UE_MODULE_LOCATION/Src"
+BUILD_LOCATION="$UE_MODULE_LOCATION/Intermediate"
+OUTPUT_LOCATION="$UE_MODULE_LOCATION/Lib"
+
+if [[ ! -w $OUTPUT_LOCATION/libbinka_ue_decode_ios_static.a ]]; then
+	echo "Check out $OUTPUT_LOCATION/libbinka_ue_decode_ios_static.a"
+	exit 1
+fi
+
+rm -rf $BUILD_LOCATION
+mkdir $BUILD_LOCATION
+mkdir $BUILD_LOCATION/ios
+mkdir $BUILD_LOCATION/ios/arm64
+mkdir $BUILD_LOCATION/ios/arm7
+mkdir $BUILD_LOCATION/ios/arm7s
+mkdir $BUILD_LOCATION/ios/x64
+
+SIM_ARGS=(
+	-arch x86_64
+	--target=x86_64-apple-ios10
+	-mmmx
+	-isysroot /Applications/Xcode.app/Contents/Developer/Platforms/iPhoneSimulator.platform/Developer/SDKs/iPhoneSimulator.sdk
+)
+
+ARM64_ARGS=(
+	-arch arm64
+	--target=arm-apple-ios10
+	-isysroot /Applications/Xcode.app/Contents/Developer/Platforms/iPhoneOS.platform/Developer/SDKs/iPhoneOS.sdk
+)
+
+ARM7_ARGS=(
+	-arch armv7
+	--target=arm-apple-ios10
+	-isysroot /Applications/Xcode.app/Contents/Developer/Platforms/iPhoneOS.platform/Developer/SDKs/iPhoneOS.sdk
+)
+ARM7S_ARGS=(
+	-arch armv7s
+	--target=arm-apple-ios10
+	-isysroot /Applications/Xcode.app/Contents/Developer/Platforms/iPhoneOS.platform/Developer/SDKs/iPhoneOS.sdk
+)
+COMMON_ARGS=(
+	-c
+	-D__GCC__
+	-DNDEBUG
+	-D__RADINSTATICLIB__
+	-DWRAP_PUBLICS=BACDUE # Prefix symbols so that mulitple libs using the same source dont get reused by the linker
+	-ffp-contract=off # Prevent FMA contraction for consistency between arm/x64
+	-fno-exceptions
+	-fno-omit-frame-pointer
+	-fno-rtti
+	-fno-strict-aliasing # prevent optimizations from introducing random silent bugs
+	-fno-unroll-loops
+	-fno-vectorize # Weve vectorized everything so this just makes the tail computation get unrolled unnecessarily
+	-fvisibility=hidden
+	-ggdb
+	-mllvm
+	-inline-threshold=64 # Pass inline threhold to llvm to prevent binary size explosion
+	-momit-leaf-frame-pointer
+	-O2
+	-IInclude
+)
+
+SOURCES=(
+	binkacd.cpp
+	binka_ue_decode.cpp
+	radfft.cpp
+)
+
+
+OUTPUT=()
+for source_file in "${SOURCES[@]}"
+do
+	clang ${ARM64_ARGS[@]} ${COMMON_ARGS[@]} -o $BUILD_LOCATION/ios/arm64/${source_file%.*}.o $SOURCE_LOCATION/$source_file
+	OUTPUT+=( $BUILD_LOCATION/ios/arm64/${source_file%.*}.o )
+done
+ar rcs $OUTPUT_LOCATION/libbinka_ue_decode_ios_arm64_static.a ${OUTPUT[@]}
+
+OUTPUT=()
+for source_file in "${SOURCES[@]}"
+do
+	clang ${ARM7_ARGS[@]} ${COMMON_ARGS[@]} -o $BUILD_LOCATION/ios/arm7/${source_file%.*}.o $SOURCE_LOCATION/$source_file
+	OUTPUT+=( $BUILD_LOCATION/ios/arm7/${source_file%.*}.o )
+done
+ar rcs $OUTPUT_LOCATION/libbinka_ue_decode_ios_arm7_static.a ${OUTPUT[@]}
+
+OUTPUT=()
+for source_file in "${SOURCES[@]}"
+do
+	clang ${ARM7S_ARGS[@]} ${COMMON_ARGS[@]} -o $BUILD_LOCATION/ios/arm7s/${source_file%.*}.o $SOURCE_LOCATION/$source_file
+	OUTPUT+=( $BUILD_LOCATION/ios/arm7s/${source_file%.*}.o )
+done
+ar rcs $OUTPUT_LOCATION/libbinka_ue_decode_ios_arm7s_static.a ${OUTPUT[@]}
+OUTPUT=()
+
+SOURCES+=(x86_cpu.c)
+for source_file in "${SOURCES[@]}"
+do
+	clang ${SIM_ARGS[@]} ${COMMON_ARGS[@]} -o $BUILD_LOCATION/ios/x64/${source_file%.*}.o $SOURCE_LOCATION/$source_file
+	OUTPUT+=( $BUILD_LOCATION/ios/x64/${source_file%.*}.o )
+done
+
+ar rcs $OUTPUT_LOCATION/libbinka_ue_decode_ios_x64_static.a ${OUTPUT[@]}
+
+lipo -create $OUTPUT_LOCATION/libbinka_ue_decode_ios_x64_static.a $OUTPUT_LOCATION/libbinka_ue_decode_ios_arm64_static.a $OUTPUT_LOCATION/libbinka_ue_decode_ios_arm7_static.a $OUTPUT_LOCATION/libbinka_ue_decode_ios_arm7s_static.a -output $OUTPUT_LOCATION/libbinka_ue_decode_ios_static.a
+rm $OUTPUT_LOCATION/libbinka_ue_decode_ios_x64_static.a $OUTPUT_LOCATION/libbinka_ue_decode_ios_arm64_static.a $OUTPUT_LOCATION/libbinka_ue_decode_ios_arm7_static.a $OUTPUT_LOCATION/libbinka_ue_decode_ios_arm7s_static.a
+
+echo Done.
diff --git a/Engine/Source/Runtime/BinkAudioDecoder/SDK/BinkAudio/BuildDecoderMac.command b/Engine/Source/Runtime/BinkAudioDecoder/SDK/BinkAudio/BuildDecoderMac.command
new file mode 100755
--- /dev/null
+++ b/Engine/Source/Runtime/BinkAudioDecoder/SDK/BinkAudio/BuildDecoderMac.command
@@ -0,0 +1,85 @@
+#!/bin/bash
+
+set -e
+
+UE_MODULE_LOCATION=`pwd`
+
+SOURCE_LOCATION="$UE_MODULE_LOCATION/Src"
+BUILD_LOCATION="$UE_MODULE_LOCATION/Intermediate"
+OUTPUT_LOCATION="$UE_MODULE_LOCATION/Lib"
+
+if [[ ! -w $OUTPUT_LOCATION/libbinka_ue_decode_osx_static.a ]]; then
+	echo "Check out $OUTPUT_LOCATION/libbinka_ue_decode_osx_static.a"
+	exit 1
+fi
+
+rm -rf $BUILD_LOCATION
+mkdir $BUILD_LOCATION
+mkdir $BUILD_LOCATION/arm64
+mkdir $BUILD_LOCATION/x64
+
+X64_ARGS=(
+	-arch x86_64
+	-msse
+	-msse2
+	-msse3
+	-mssse3
+)
+
+ARM64_ARGS=(
+	-arch arm64
+)
+
+COMMON_ARGS=(
+	-c
+	-D__GCC__
+	-D__MACHO__
+	-DNDEBUG
+	-D__RADINSTATICLIB__
+	-DWRAP_PUBLICS=BACDUE # Prefix symbols so that mulitple libs using the same source dont get reused by the linker
+	-ffp-contract=off # Prevent FMA contraction for consistency between arm/x64
+	-fno-exceptions
+	-fno-omit-frame-pointer
+	-fno-rtti
+	-fno-strict-aliasing # prevent optimizations from introducing random silent bugs
+	-fno-unroll-loops
+	-fno-vectorize # Weve vectorized everything so this just makes the tail computation get unrolled unnecessarily
+	-fvisibility=hidden
+	-ggdb
+	-mllvm
+	-inline-threshold=64 # Pass inline threhold to llvm to prevent binary size explosion
+	-momit-leaf-frame-pointer
+	-O2
+	-IInclude
+	-mmacosx-version-min=10.14
+)
+
+SOURCES=(
+	binkacd.cpp
+	x86_cpu.c
+	binka_ue_decode.cpp
+	radfft.cpp
+)
+
+
+for source_file in "${SOURCES[@]}"
+do
+	clang ${X64_ARGS[@]} ${COMMON_ARGS[@]} -o $BUILD_LOCATION/x64/${source_file%.*}.o $SOURCE_LOCATION/$source_file
+	OUTPUT+=( $BUILD_LOCATION/x64/${source_file%.*}.o )
+done
+
+ar rcs $OUTPUT_LOCATION/libbinka_ue_decode_osx_x64_static.a ${OUTPUT[@]}
+
+OUTPUT=()
+for source_file in "${SOURCES[@]}"
+do
+	clang ${ARM64_ARGS[@]} ${COMMON_ARGS[@]} -o $BUILD_LOCATION/arm64/${source_file%.*}.o $SOURCE_LOCATION/$source_file
+	OUTPUT+=( $BUILD_LOCATION/arm64/${source_file%.*}.o )
+done
+ar rcs $OUTPUT_LOCATION/libbinka_ue_decode_osx_arm64_static.a ${OUTPUT[@]}
+
+lipo -create $OUTPUT_LOCATION/libbinka_ue_decode_osx_x64_static.a $OUTPUT_LOCATION/libbinka_ue_decode_osx_arm64_static.a -output $OUTPUT_LOCATION/libbinka_ue_decode_osx_static.a
+rm $OUTPUT_LOCATION/libbinka_ue_decode_osx_arm64_static.a
+rm $OUTPUT_LOCATION/libbinka_ue_decode_osx_x64_static.a
+
+echo Done.
diff --git a/Engine/Source/Runtime/BinkAudioDecoder/SDK/BinkAudio/BuildEncoderMac.command b/Engine/Source/Runtime/BinkAudioDecoder/SDK/BinkAudio/BuildEncoderMac.command
new file mode 100755
--- /dev/null
+++ b/Engine/Source/Runtime/BinkAudioDecoder/SDK/BinkAudio/BuildEncoderMac.command
@@ -0,0 +1,84 @@
+#!/bin/bash
+
+set -e
+
+UE_MODULE_LOCATION=`pwd`
+
+SOURCE_LOCATION="$UE_MODULE_LOCATION/Src"
+BUILD_LOCATION="$UE_MODULE_LOCATION/Intermediate"
+OUTPUT_LOCATION="$UE_MODULE_LOCATION/Lib"
+if [[ ! -w $OUTPUT_LOCATION/libbinka_ue_encode_osx_static.a ]]; then
+	echo "Check out $OUTPUT_LOCATION/libbinka_ue_encode_osx_static.a"
+	exit 1
+fi
+
+rm -rf $BUILD_LOCATION
+mkdir $BUILD_LOCATION
+mkdir $BUILD_LOCATION/arm64
+mkdir $BUILD_LOCATION/x64
+
+X64_ARGS=(
+	-arch x86_64
+	-msse
+	-msse2
+	-msse3
+	-mssse3
+)
+
+ARM64_ARGS=(
+	-arch arm64
+)
+
+COMMON_ARGS=(
+	-c
+	-D__GCC__
+	-D__MACHO__
+	-DNDEBUG
+	-D__RADINSTATICLIB__
+	-DWRAP_PUBLICS=BACEUE # Prefix symbols so that mulitple libs using the same source dont get reused by the linker
+	-ffp-contract=off # Prevent FMA contraction for consistency between arm/x64
+	-fno-exceptions
+	-fno-omit-frame-pointer
+	-fno-rtti
+	-fno-strict-aliasing # prevent optimizations from introducing random silent bugs
+	-fno-unroll-loops
+	-fno-vectorize # Weve vectorized everything so this just makes the tail computation get unrolled unnecessarily
+	-fvisibility=hidden
+	-ggdb
+	-mllvm
+	-inline-threshold=64 # Pass inline threhold to llvm to prevent binary size explosion
+	-momit-leaf-frame-pointer
+	-O2
+	-IInclude
+	-mmacosx-version-min=10.14
+)
+
+SOURCES=(
+	binkace.c
+	popmal.c
+	varbits.c
+	binka_ue_encode.cpp
+	radfft.cpp
+	ranged_log.cpp
+)
+
+
+for source_file in "${SOURCES[@]}"
+do
+	clang ${X64_ARGS[@]} ${COMMON_ARGS[@]} -o $BUILD_LOCATION/x64/${source_file%.*}.o $SOURCE_LOCATION/$source_file
+	OUTPUT+=( $BUILD_LOCATION/x64/${source_file%.*}.o )
+done
+
+ar rcs $OUTPUT_LOCATION/libbinka_ue_encode_osx_x64_static.a ${OUTPUT[@]}
+
+OUTPUT=()
+for source_file in "${SOURCES[@]}"
+do
+	clang ${ARM64_ARGS[@]} ${COMMON_ARGS[@]} -o $BUILD_LOCATION/arm64/${source_file%.*}.o $SOURCE_LOCATION/$source_file
+	OUTPUT+=( $BUILD_LOCATION/arm64/${source_file%.*}.o )
+done
+ar rcs $OUTPUT_LOCATION/libbinka_ue_encode_osx_arm64_static.a ${OUTPUT[@]}
+
+lipo -create $OUTPUT_LOCATION/libbinka_ue_encode_osx_x64_static.a $OUTPUT_LOCATION/libbinka_ue_encode_osx_arm64_static.a -output $OUTPUT_LOCATION/libbinka_ue_encode_osx_static.a
+rm $OUTPUT_LOCATION/libbinka_ue_encode_osx_x64_static.a $OUTPUT_LOCATION/libbinka_ue_encode_osx_arm64_static.a 
+echo Done.
diff --git a/Engine/Source/Runtime/BinkAudioDecoder/SDK/BinkAudio/BuildForMac.command b/Engine/Source/Runtime/BinkAudioDecoder/SDK/BinkAudio/BuildForMac.command
deleted file mode 100755
diff --git a/Engine/Source/Runtime/BinkAudioDecoder/SDK/BinkAudio/CMakeLists.txt b/Engine/Source/Runtime/BinkAudioDecoder/SDK/BinkAudio/CMakeLists.txt
deleted file mode 100644
diff --git a/Engine/Source/Runtime/BinkAudioDecoder/SDK/BinkAudio/HowToBuild.txt b/Engine/Source/Runtime/BinkAudioDecoder/SDK/BinkAudio/HowToBuild.txt
--- a/Engine/Source/Runtime/BinkAudioDecoder/SDK/BinkAudio/HowToBuild.txt
+++ b/Engine/Source/Runtime/BinkAudioDecoder/SDK/BinkAudio/HowToBuild.txt
@@ -3,11 +3,12 @@ Bink Audio Decoder/Encoder Build Instructions
 First, the necessary source files to pass to the compiler:
 
 Encoder
-    popmal.c
     binkace.c
+    popmal.c
+    varbits.c
     binka_ue_encode.cpp
     radfft.cpp
-    varbits.c
+    ranged_log.cpp
 
 Decoder
     binkacd.cpp
diff --git a/Engine/Source/Runtime/BinkAudioDecoder/SDK/BinkAudio/Src/binka_ue_decode_test.cpp b/Engine/Source/Runtime/BinkAudioDecoder/SDK/BinkAudio/Src/binka_ue_decode_test.cpp
--- a/Engine/Source/Runtime/BinkAudioDecoder/SDK/BinkAudio/Src/binka_ue_decode_test.cpp
+++ b/Engine/Source/Runtime/BinkAudioDecoder/SDK/BinkAudio/Src/binka_ue_decode_test.cpp
@@ -38,7 +38,12 @@ CONSOLE_MAIN
     int BlockIndex = 0;
     for (; InputBuffer < InputBufferEnd; BlockIndex++)
     {
-        short OutputData[1920];
+        // bink audio requires all buffers be 16 byte aligned.
+#if defined(__clang__)
+        short __attribute__ ((aligned (16))) OutputData[1920];
+#else
+        short __declspec(align(16))  OutputData[1920];
+#endif
 
         U16 BlockSize = *(U16*)InputBuffer;
         InputBuffer += 2;
@@ -86,3 +91,4 @@ CONSOLE_MAIN
 
     return Failed ? 1 : 0;
 }
+
diff --git a/Engine/Source/Runtime/BinkAudioDecoder/SDK/BinkAudio/Src/binka_ue_encode.cpp b/Engine/Source/Runtime/BinkAudioDecoder/SDK/BinkAudio/Src/binka_ue_encode.cpp
--- a/Engine/Source/Runtime/BinkAudioDecoder/SDK/BinkAudio/Src/binka_ue_encode.cpp
+++ b/Engine/Source/Runtime/BinkAudioDecoder/SDK/BinkAudio/Src/binka_ue_encode.cpp
@@ -31,6 +31,81 @@ struct MemBuffer
     U32 total_bytes;
 };
 
+static unsigned int CRC32Table[] =
+{
+    0x00000000, 0x77073096, 0xEE0E612C, 0x990951BA, 0x076DC419,
+    0x706AF48F, 0xE963A535, 0x9E6495A3, 0x0EDB8832, 0x79DCB8A4,
+    0xE0D5E91E, 0x97D2D988, 0x09B64C2B, 0x7EB17CBD, 0xE7B82D07,
+    0x90BF1D91, 0x1DB71064, 0x6AB020F2, 0xF3B97148, 0x84BE41DE,
+    0x1ADAD47D, 0x6DDDE4EB, 0xF4D4B551, 0x83D385C7, 0x136C9856,
+    0x646BA8C0, 0xFD62F97A, 0x8A65C9EC, 0x14015C4F, 0x63066CD9,
+    0xFA0F3D63, 0x8D080DF5, 0x3B6E20C8, 0x4C69105E, 0xD56041E4,
+    0xA2677172, 0x3C03E4D1, 0x4B04D447, 0xD20D85FD, 0xA50AB56B,
+    0x35B5A8FA, 0x42B2986C, 0xDBBBC9D6, 0xACBCF940, 0x32D86CE3,
+    0x45DF5C75, 0xDCD60DCF, 0xABD13D59, 0x26D930AC, 0x51DE003A,
+    0xC8D75180, 0xBFD06116, 0x21B4F4B5, 0x56B3C423, 0xCFBA9599,
+    0xB8BDA50F, 0x2802B89E, 0x5F058808, 0xC60CD9B2, 0xB10BE924,
+    0x2F6F7C87, 0x58684C11, 0xC1611DAB, 0xB6662D3D, 0x76DC4190,
+    0x01DB7106, 0x98D220BC, 0xEFD5102A, 0x71B18589, 0x06B6B51F,
+    0x9FBFE4A5, 0xE8B8D433, 0x7807C9A2, 0x0F00F934, 0x9609A88E,
+    0xE10E9818, 0x7F6A0DBB, 0x086D3D2D, 0x91646C97, 0xE6635C01,
+    0x6B6B51F4, 0x1C6C6162, 0x856530D8, 0xF262004E, 0x6C0695ED,
+    0x1B01A57B, 0x8208F4C1, 0xF50FC457, 0x65B0D9C6, 0x12B7E950,
+    0x8BBEB8EA, 0xFCB9887C, 0x62DD1DDF, 0x15DA2D49, 0x8CD37CF3,
+    0xFBD44C65, 0x4DB26158, 0x3AB551CE, 0xA3BC0074, 0xD4BB30E2,
+    0x4ADFA541, 0x3DD895D7, 0xA4D1C46D, 0xD3D6F4FB, 0x4369E96A,
+    0x346ED9FC, 0xAD678846, 0xDA60B8D0, 0x44042D73, 0x33031DE5,
+    0xAA0A4C5F, 0xDD0D7CC9, 0x5005713C, 0x270241AA, 0xBE0B1010,
+    0xC90C2086, 0x5768B525, 0x206F85B3, 0xB966D409, 0xCE61E49F,
+    0x5EDEF90E, 0x29D9C998, 0xB0D09822, 0xC7D7A8B4, 0x59B33D17,
+    0x2EB40D81, 0xB7BD5C3B, 0xC0BA6CAD, 0xEDB88320, 0x9ABFB3B6,
+    0x03B6E20C, 0x74B1D29A, 0xEAD54739, 0x9DD277AF, 0x04DB2615,
+    0x73DC1683, 0xE3630B12, 0x94643B84, 0x0D6D6A3E, 0x7A6A5AA8,
+    0xE40ECF0B, 0x9309FF9D, 0x0A00AE27, 0x7D079EB1, 0xF00F9344,
+    0x8708A3D2, 0x1E01F268, 0x6906C2FE, 0xF762575D, 0x806567CB,
+    0x196C3671, 0x6E6B06E7, 0xFED41B76, 0x89D32BE0, 0x10DA7A5A,
+    0x67DD4ACC, 0xF9B9DF6F, 0x8EBEEFF9, 0x17B7BE43, 0x60B08ED5,
+    0xD6D6A3E8, 0xA1D1937E, 0x38D8C2C4, 0x4FDFF252, 0xD1BB67F1,
+    0xA6BC5767, 0x3FB506DD, 0x48B2364B, 0xD80D2BDA, 0xAF0A1B4C,
+    0x36034AF6, 0x41047A60, 0xDF60EFC3, 0xA867DF55, 0x316E8EEF,
+    0x4669BE79, 0xCB61B38C, 0xBC66831A, 0x256FD2A0, 0x5268E236,
+    0xCC0C7795, 0xBB0B4703, 0x220216B9, 0x5505262F, 0xC5BA3BBE,
+    0xB2BD0B28, 0x2BB45A92, 0x5CB36A04, 0xC2D7FFA7, 0xB5D0CF31,
+    0x2CD99E8B, 0x5BDEAE1D, 0x9B64C2B0, 0xEC63F226, 0x756AA39C,
+    0x026D930A, 0x9C0906A9, 0xEB0E363F, 0x72076785, 0x05005713,
+    0x95BF4A82, 0xE2B87A14, 0x7BB12BAE, 0x0CB61B38, 0x92D28E9B,
+    0xE5D5BE0D, 0x7CDCEFB7, 0x0BDBDF21, 0x86D3D2D4, 0xF1D4E242,
+    0x68DDB3F8, 0x1FDA836E, 0x81BE16CD, 0xF6B9265B, 0x6FB077E1,
+    0x18B74777, 0x88085AE6, 0xFF0F6A70, 0x66063BCA, 0x11010B5C,
+    0x8F659EFF, 0xF862AE69, 0x616BFFD3, 0x166CCF45, 0xA00AE278,
+    0xD70DD2EE, 0x4E048354, 0x3903B3C2, 0xA7672661, 0xD06016F7,
+    0x4969474D, 0x3E6E77DB, 0xAED16A4A, 0xD9D65ADC, 0x40DF0B66,
+    0x37D83BF0, 0xA9BCAE53, 0xDEBB9EC5, 0x47B2CF7F, 0x30B5FFE9,
+    0xBDBDF21C, 0xCABAC28A, 0x53B39330, 0x24B4A3A6, 0xBAD03605,
+    0xCDD70693, 0x54DE5729, 0x23D967BF, 0xB3667A2E, 0xC4614AB8,
+    0x5D681B02, 0x2A6F2B94, 0xB40BBE37, 0xC30C8EA1, 0x5A05DF1B,
+    0x2D02EF8D
+};
+
+// We keep this around if we need to do consistency checks for arm/x64 builds
+static U32 CRC(const void* src, int n_bytes)
+{
+    U32 C;
+
+    C = 0xFFFFFFFF;
+
+    {
+        unsigned char const *S = (unsigned char const *) src;
+        while (n_bytes--)
+        {
+            C = CRC32Table[((int) (C) ^ (*S++)) & 0xFF] ^ ((C) >> 8);
+        }
+    }
+
+    C ^= 0xFFFFFFFF;
+    return C;
+}
+
 //-----------------------------------------------------------------------------
 //-----------------------------------------------------------------------------
 static void MemBufferAdd(MemBuffer* mem, void* data, U32 data_len)
@@ -163,6 +238,95 @@ static U32 SeekTableBufferTrim(SeekTableBuffer* seek)
     return frames_per_entry;
 }
 
+// Scope guard to set up FP state as desired and reset it on exit
+struct FPStateScope
+{
+    U32 saved_state;
+
+    FPStateScope();
+    ~FPStateScope();
+};
+
+#if defined(__RADSSE2__)
+    FPStateScope::FPStateScope()
+    {
+        saved_state = _mm_getcsr();
+        // Set up our expected FP state: no exception flags set,
+        // all exceptions masked (suppressed), round to nearest,
+        // flush to zero and denormals are zero both off.
+        _mm_setcsr(_MM_MASK_MASK /* all exceptions masked */ | _MM_ROUND_NEAREST | _MM_FLUSH_ZERO_OFF);
+    }
+
+    FPStateScope::~FPStateScope()
+    {
+        _mm_setcsr(saved_state);
+    }
+
+#elif defined(__RADARM__) && defined(__RAD64__)
+
+    #ifdef _MSC_VER
+
+        #include <intrin.h>
+        static U32 read_fpcr()
+        {
+            // The system register R/W instructions use 64-bit GPRs, but the
+            // architectural FPCR is 32b
+            return (U32)_ReadStatusReg(ARM64_FPCR);
+        }
+
+        static void write_fpcr(U32 state)
+        {
+            _WriteStatusReg(ARM64_FPCR, state);
+        }
+
+    #elif defined(__clang__) || defined(__GNUC__)
+
+        static U32 read_fpcr()
+        {
+            // The system register R/W instructions use 64-bit GPRs, but the
+            // architectural FPCR is 32b
+            U64 value;
+            __asm__ volatile("mrs %0, fpcr" : "=r"(value));
+            return (U32)value;
+        }
+
+        static void write_fpcr(U32 state)
+        {
+            U64 state64 = state;
+            __asm__ volatile("msr fpcr, %0" : : "r"(state64));
+        }
+
+    #else
+
+        #error compiler? Not clang or msvc
+
+    #endif
+
+    FPStateScope::FPStateScope()
+    {
+        saved_state = read_fpcr();
+
+        // IEEE compliant mode in FPCR is just all-0
+        write_fpcr(0);
+    }
+
+    FPStateScope::~FPStateScope()
+    {
+        write_fpcr(saved_state);
+    }
+
+#else // neither SSE2 nor ARM64
+
+    FPStateScope::FPStateScope()
+        : saved_state(0)
+    {
+    }
+
+    FPStateScope::~FPStateScope()
+    {
+    }
+
+#endif
 
 //-----------------------------------------------------------------------------
 //-----------------------------------------------------------------------------
@@ -171,6 +335,8 @@ uint8_t UECompressBinkAudio(
     uint8_t Quality, uint8_t GenerateSeekTable, BAUECompressAllocFnType* MemAlloc, BAUECompressFreeFnType* MemFree,
     void** OutData, uint32_t* OutDataLen)
 {
+    FPStateScope FixFloatingPoint;
+
     //
     // Deinterlace the input.
     //
@@ -196,7 +362,7 @@ uint8_t UECompressBinkAudio(
 
         if (ChannelsPerStream[i] == 2)
         {
-            while (pInput != pEnd)
+            while (pInput < pEnd) // less than to protect against malformed wavs
             {
                 *(S32*)pOutput = *(S32*)pInput;
                 pInput += InputStride;
@@ -205,7 +371,7 @@ uint8_t UECompressBinkAudio(
         }
         else
         {
-            while (pInput != pEnd)
+            while (pInput < pEnd) // less than to protect against malformed wavs
             {
                 *(S16*)pOutput = *(S16*)pInput;
                 pInput += InputStride;
@@ -282,6 +448,8 @@ uint8_t UECompressBinkAudio(
                 OutputBuffers + StreamIndex,
                 OutputLens + StreamIndex,
                 InputUseds + StreamIndex);
+
+            //printf("Block In: %08x Out: %08x\n", CRC(InputBuffers[StreamIndex], InputLens[StreamIndex]), CRC(OutputBuffers[StreamIndex], OutputLens[StreamIndex]));
         }
 
         // Update cursors.
@@ -353,8 +521,9 @@ uint8_t UECompressBinkAudio(
 
     BinkAudioFileHeader Header;
     Header.tag = 'UEBA';
+    Header.PADDING = 0;
     Header.channels = (U8)WavChannels;
-    Header.rate = WavRate;
+    Header.rate = (U16)WavRate;
     Header.sample_count = SamplesPerChannel;
     Header.max_comp_space_needed = (U16)MaxBlockSize;
     Header.flags = 1;
diff --git a/Engine/Source/Runtime/BinkAudioDecoder/SDK/BinkAudio/Src/binkacd.cpp b/Engine/Source/Runtime/BinkAudioDecoder/SDK/BinkAudio/Src/binkacd.cpp
--- a/Engine/Source/Runtime/BinkAudioDecoder/SDK/BinkAudio/Src/binkacd.cpp
+++ b/Engine/Source/Runtime/BinkAudioDecoder/SDK/BinkAudio/Src/binkacd.cpp
@@ -452,7 +452,7 @@ RAD_USES_SSSE3 static RADINLINE void decode_pfxsum_8(float* results, const uint8
     }
 }
 
-static RADINLINE size_t decode_pfxsum(float* results, const uint8_t* buffer, size_t bitpos_vals, size_t bitpos_signs, size_t count, int bitlen)
+static RADINLINE size_t decode_pfxsum(float* results, const uint8_t* buffer, size_t bitpos_vals, size_t bitpos_signs, size_t end_bit_position, size_t count, int bitlen)
 {
     Vec128 last_prefix_sum;
     size_t i;
@@ -478,8 +478,28 @@ static RADINLINE size_t decode_pfxsum(float* results, const uint8_t* buffer, siz
 
         // Advance read cursors
         bitpos_vals += 8 * bitlen;
-        nsigns = (size_t)(unsigned int)_mm_extract_epi16(last_prefix_sum, 7);
+
+        //
+        // This extract doesn't actually need the &15, as the prefix sum can only ever be
+        // a max of 8, however we have crash reports in the wild of situations where bitpos_signs
+        // has gone stupid high and causing an AV. This happens after the initial run, and is _well_
+        // off in to another page, so the belief is that we're getting some CPU overclocking
+        // absurdity. This is confirmed because in the crash minidump we can see that the prefix
+        // sum is sane. Since we don't know where the issue is, we're adding the &15. This ends up
+        // replacing the mov ecx, ecx for sign extension in a 1:1 trade, so it's completely free.
+        //
+        // Then, to help sanitize, we test the bit positions against the end of the buffer and
+        // just bail since we're off the edge of the map anyway. This ends up getting caught by
+        // the bit position validation in read_channel_data_2 and zeroing the block.
+        //
+        nsigns = (size_t)(unsigned int)(_mm_extract_epi16(last_prefix_sum, 7) & 0xf);
         bitpos_signs += nsigns;
+
+        if (bitpos_signs > end_bit_position ||
+            bitpos_vals > end_bit_position)
+        {
+            return end_bit_position;
+        }
     }
 
     // this should almost never happen - only when a run passes the end of the transform size.
@@ -742,7 +762,7 @@ read_chan_data_corrupted:
         {
           if (validate_chunked_bit_read<Vec128>(bit_position, end_bit_position, coeff_bits_needed) == false)
             goto read_chan_data_corrupted;
-          bit_position = decode_pfxsum(samps + i, buffer, bit_position, sign_bit_position, end - i, bitlen);
+          bit_position = decode_pfxsum(samps + i, buffer, bit_position, sign_bit_position, end_bit_position, end - i, bitlen);
         }
         else
 #endif
diff --git a/Engine/Source/Runtime/BinkAudioDecoder/SDK/BinkAudio/Src/binkace.c b/Engine/Source/Runtime/BinkAudioDecoder/SDK/BinkAudio/Src/binkace.c
--- a/Engine/Source/Runtime/BinkAudioDecoder/SDK/BinkAudio/Src/binkace.c
+++ b/Engine/Source/Runtime/BinkAudioDecoder/SDK/BinkAudio/Src/binkace.c
@@ -364,9 +364,9 @@ static AUDIOFLOAT calctonality( U32 transform_size,
     if ( c > 0.5F )
       c = 0.5F;
     else if ( c < 0.05F)
-      c = 0.5F;
+      c = 0.05F;
 
-    c = (AUDIOFLOAT)( -0.43F * radlog( c ) ) - 0.29F;
+    c = (AUDIOFLOAT)( -0.43F * ranged_log_0p05_to_0p5( c ) ) - 0.29F;
 
     t += c * m;
     n += m;
diff --git a/Engine/Source/Runtime/BinkAudioDecoder/SDK/BinkAudio/Src/radmath.h b/Engine/Source/Runtime/BinkAudioDecoder/SDK/BinkAudio/Src/radmath.h
--- a/Engine/Source/Runtime/BinkAudioDecoder/SDK/BinkAudio/Src/radmath.h
+++ b/Engine/Source/Runtime/BinkAudioDecoder/SDK/BinkAudio/Src/radmath.h
@@ -10,69 +10,24 @@ RADDEFSTART
 #define mult64andshift( a, b, c )  ( (U32) ( ( ( (U64) a ) * ( (U64) b ) ) >> ( (U64) c ) ) )
 #define radabs abs
 
-#if defined( _MSC_VER )
-
-  #if defined(__RADX86__) 
-    #pragma intrinsic(abs, log, fabs, sqrt, fmod, sin, cos, tan, asin, acos, atan, atan2, exp)
-  #endif
-
-  #define radcos( val ) cos( (float)(val) )
-  #define radsin( val ) sin( (float)(val) )
-  #define radatan( val ) atan( (float)(val) )
-  #define radatan2( val1, val2 ) atan2( (float)(val1), (float)(val2) )
-  #define radpow( val1, val2 ) pow( (float)(val1), (float)(val2) )
-  #define radfsqrt( val ) sqrt( (float)(val) )
-  #define radlog( val ) log( (float)(val) )
-  #define radlog10( val ) log10( (float)(val) )
-  #define radexp( val ) exp( (float)(val) )
-  #define radfabs( val ) fabs( (float)(val) )
-  #define radfloor( val ) floor( (float)(val) )
-  #define radceil( val ) ceil( (float)(val) )
-
-  #if defined(__RAD64__) || (__RADARM__)
-  
-    #define LN2 0.693147181F
-
-    static F64 __inline radlog2( F64 X )
-    {
-      return( radlog( X ) / LN2 );
-    }
-
-  #else
-  
-    static F64 __inline radlog2( F64 X )
-    {
-      F64 Result;
-      __asm
-      {
-        fld1
-        fld X
-        fyl2x
-        fstp Result
-      }
-      return( Result );
-    }
-  
-  #endif
-
-#else
+// We use the floating point version of everything and have casts around everything
+// to ensure we don't accidentally have an expression evaluate as double on one platform
+// and float on the other.
+//
+// This was due to previously using e.g. sqrt on one platform and sqrtf on another, but there's
+// no reason to _not_ have these casts.
+#define radatan2( val1, val2 )    ((float)atan2f  ( (float)val1, (float)val2) )
+#define radpow( val1, val2 )      ((float)powf    ( (float)val1, (float)val2) )
+#define radfsqrt( val )           ((float)sqrtf   ( (float)val) )
+#define radlog10( val )           ((float)log10f  ( (float)val) )
+#define radexp( val )             ((float)expf    ( (float)val) )
+#define radfabs( val )            ((float)fabsf   ( (float)val) )
+#define radfloor( val )           ((float)floorf  ( (float)val) )
+
+float ranged_log_0p05_to_0p5(float x);
+
+#if !defined( _MSC_VER )
   #include <stdlib.h> // abs
-    
-  // generally no non-windows math stuff...
-    
-  //#define radcos( val ) cosf( val )
-  //#define radsin( val ) sinf( val )
-  //#define radatan( val ) atanf( val )
-  #define radatan2( val1, val2 ) atan2f( val1, val2 )
-  #define radpow( val1, val2 ) powf( val1, val2 )
-  #define radfsqrt( val ) sqrtf( val )
-  #define radlog( val ) log( val )
-  #define radlog10( val ) log10f( val )
-  #define radexp( val ) expf( val )
-  #define radfabs( val ) fabsf( val )
-  #define radfloor( val ) floorf( val )
-  //#define radceil( val ) ceilf( val )
-
 #endif
 
 RADDEFEND
diff --git a/Engine/Source/Runtime/BinkAudioDecoder/SDK/BinkAudio/Src/ranged_log.cpp b/Engine/Source/Runtime/BinkAudioDecoder/SDK/BinkAudio/Src/ranged_log.cpp
new file mode 100644
--- /dev/null
+++ b/Engine/Source/Runtime/BinkAudioDecoder/SDK/BinkAudio/Src/ranged_log.cpp
@@ -0,0 +1,17 @@
+
+//
+// In order to avoid potential consistency issues between libms we have this narrow-range log function
+// as we only use it in one place inside binkace and we only use it within a small range.
+//
+// 
+// handles 0.05 to 0.5 (can easily be full range with an if and no table for the ln2*scale)
+extern "C" float ranged_log_0p05_to_0p5(float x) 
+{
+	static float twotoe[5] = { 4.0f*0.69314718f, 3.0f*0.69314718f, 2.0f*0.69314718f, 1.0f*0.69314718f, 0.0f };  // handles range of 0.05 to 0.5 only needs 5 entryies (final value zero is just for 0.5)
+	static int k5 = 0x3fdef298; // close to -1.7417939f, but scanned to minimize max error (which is 0.000061170932 with ave of 0.00003781923)
+	union { int i; float f; } f;
+	f.f = x;
+	f.i &= 0xff800000; // isolates the mantissa and sign
+	x = f.f / x;  // gets a less-than-1.0 value into range of 0.5 to 1
+	return *((float*)&k5) + ( -5.64234426f + ( 5.87976061f + ( -3.57742942f + 0.905148439f * x) * x) * x) * x - twotoe[ (f.i>>23) - 122 ]; // constants from maple, using remez
+}
diff --git a/Engine/Source/Runtime/BinkAudioDecoder/SDK/BinkAudio/Src/varbits.c b/Engine/Source/Runtime/BinkAudioDecoder/SDK/BinkAudio/Src/varbits.c
--- a/Engine/Source/Runtime/BinkAudioDecoder/SDK/BinkAudio/Src/varbits.c
+++ b/Engine/Source/Runtime/BinkAudioDecoder/SDK/BinkAudio/Src/varbits.c
@@ -16,7 +16,6 @@ RADDEFINEDATA const RAD_ALIGN(U32, VarBitsLens[ 33 ], 32 ) =
   0xffffffff
 };
 
-#if defined(__RADX86__)
 RADDEFFUNC void VarBitsCopy(VARBITS* dest,VARBITS* src,U32 size)
 {
   U32 val;
@@ -32,9 +31,4 @@ RADDEFFUNC void VarBitsCopy(VARBITS* dest,VARBITS* src,U32 size)
   }
 }
 
-#else
-
-char VarbitsLibtoolPacifier;
-
-#endif
 
