Copyright boilerplate

diff --git a/Engine/Source/Runtime/BinkAudioDecoder/SDK/BinkAudio/Src/ranged_log.cpp b/Engine/Source/Runtime/BinkAudioDecoder/SDK/BinkAudio/Src/ranged_log.cpp
--- a/Engine/Source/Runtime/BinkAudioDecoder/SDK/BinkAudio/Src/ranged_log.cpp
+++ b/Engine/Source/Runtime/BinkAudioDecoder/SDK/BinkAudio/Src/ranged_log.cpp
@@ -1,4 +1,4 @@
-
+// Copyright Epic Games, Inc. All Rights Reserved.
 //
 // In order to avoid potential consistency issues between libms we have this narrow-range log function
 // as we only use it in one place inside binkace and we only use it within a small range.
