Add tvOS to the allowed platform Bink Audio Decoder list and include a tvOS compatible static lib.

diff --git a/Engine/Source/Runtime/BinkAudioDecoder/Module/BinkAudioDecoder.Build.cs b/Engine/Source/Runtime/BinkAudioDecoder/Module/BinkAudioDecoder.Build.cs
--- a/Engine/Source/Runtime/BinkAudioDecoder/Module/BinkAudioDecoder.Build.cs
+++ b/Engine/Source/Runtime/BinkAudioDecoder/Module/BinkAudioDecoder.Build.cs
@@ -26,6 +26,10 @@ public class BinkAudioDecoder : ModuleRules
         {
             return Path.Combine(ModuleDirectory, "..", "SDK", "BinkAudio", "Lib", "libbinka_ue_decode_ios_static.a");
         }
+        if (Target.Platform == UnrealTargetPlatform.TVOS)
+        {
+            return Path.Combine(ModuleDirectory, "..", "SDK", "BinkAudio", "Lib", "libbinka_ue_decode_tvos_static.a");
+        }
         return null;
     }
 
