Fixes from the fuzz testing of bink audio encoder

diff --git a/Engine/Source/Runtime/BinkAudioDecoder/SDK/BinkAudio/Src/binka_ue_encode.cpp b/Engine/Source/Runtime/BinkAudioDecoder/SDK/BinkAudio/Src/binka_ue_encode.cpp
--- a/Engine/Source/Runtime/BinkAudioDecoder/SDK/BinkAudio/Src/binka_ue_encode.cpp
+++ b/Engine/Source/Runtime/BinkAudioDecoder/SDK/BinkAudio/Src/binka_ue_encode.cpp
@@ -13,6 +13,11 @@
 
 #define BINKA_COMPRESS_SUCCESS 0
 #define BINKA_COMPRESS_ERROR_CHANS 1
+#define BINKA_COMPRESS_ERROR_SAMPLES 2
+#define BINKA_COMPRESS_ERROR_RATE 3
+#define BINKA_COMPRESS_ERROR_QUALITY 4
+#define BINKA_COMPRESS_ERROR_ALLOCATORS 5
+#define BINKA_COMPRESS_ERROR_OUTPUT 6
 
 
 struct MemBufferEntry
@@ -337,11 +342,26 @@ uint8_t UECompressBinkAudio(
 {
     FPStateScope FixFloatingPoint;
 
+    if (WavChannels == 0)
+        return BINKA_COMPRESS_ERROR_CHANS;
+    if (WavRate > 256000 || WavRate < 11025)
+        return BINKA_COMPRESS_ERROR_RATE;
+    if (Quality > 9)
+        return BINKA_COMPRESS_ERROR_QUALITY;
+    if (MemAlloc == nullptr ||
+        MemFree == nullptr)
+        return BINKA_COMPRESS_ERROR_ALLOCATORS;
+    if (OutData == nullptr ||
+        OutDataLen == nullptr)
+        return BINKA_COMPRESS_ERROR_OUTPUT;
+
     //
     // Deinterlace the input.
     //
     U32 SamplesPerChannel = WavDataLen / (sizeof(S16) * WavChannels);
     U8 NumBinkStreams = (WavChannels / 2) + (WavChannels  & 1);
+    if (SamplesPerChannel == 0)
+        return BINKA_COMPRESS_ERROR_SAMPLES;
     if (NumBinkStreams > MAX_STREAMS)
         return BINKA_COMPRESS_ERROR_CHANS;
 
@@ -357,12 +377,13 @@ uint8_t UECompressBinkAudio(
 
         S32 InputStride = sizeof(S16) * WavChannels;
         char* pInput = (char*)WavData + CurrentDeintChannel*sizeof(S16);
-        char* pEnd = pInput + WavDataLen;
         char* pOutput = SourceStreams[i];
+        char* pInputEnd = pInput + SamplesPerChannel*InputStride;
+        
 
         if (ChannelsPerStream[i] == 2)
         {
-            while (pInput < pEnd) // less than to protect against malformed wavs
+            while (pInput < pInputEnd) // less than to protect against malformed wavs
             {
                 *(S32*)pOutput = *(S32*)pInput;
                 pInput += InputStride;
@@ -371,7 +392,7 @@ uint8_t UECompressBinkAudio(
         }
         else
         {
-            while (pInput < pEnd) // less than to protect against malformed wavs
+            while (pInput < pInputEnd) // less than to protect against malformed wavs
             {
                 *(S16*)pOutput = *(S16*)pInput;
                 pInput += InputStride;
@@ -552,4 +573,38 @@ uint8_t UECompressBinkAudio(
     return BINKA_COMPRESS_SUCCESS;
 }
 
+#if 0
+#include <malloc.h>
+
+struct fuzzer_input
+{
+    uint32_t Rate;
+    uint8_t Chans;
+    uint8_t Quality;
+    uint8_t GenTable;    
+};
+
+static void* AllocThunk(uintptr_t ByteCount) { return malloc(ByteCount); }
+static void FreeThunk(void* Ptr) { free(Ptr); }
+
+extern "C" int LLVMFuzzerTestOneInput(const uint8_t *Data, size_t Size) 
+{    
+    if (Size <= sizeof(fuzzer_input))
+        return 0;
+    if (Size > 0xFFFFFFFFU)
+        return 0;
 
+    fuzzer_input* input = (fuzzer_input*)Data;
+
+    void* CompressedData = 0;
+    uint32_t CompressedLen;
+
+    UECompressBinkAudio((void*)(Data + sizeof(fuzzer_input)), (uint32_t)Size - sizeof(fuzzer_input),
+        input->Rate, input->Chans, input->Quality, input->GenTable, AllocThunk, FreeThunk,
+        &CompressedData, &CompressedLen);
+
+    FreeThunk(CompressedData);
+
+    return 0;
+}
+#endif
diff --git a/Engine/Source/Runtime/BinkAudioDecoder/SDK/BinkAudio/Src/binkace.c b/Engine/Source/Runtime/BinkAudioDecoder/SDK/BinkAudio/Src/binkace.c
--- a/Engine/Source/Runtime/BinkAudioDecoder/SDK/BinkAudio/Src/binkace.c
+++ b/Engine/Source/Runtime/BinkAudioDecoder/SDK/BinkAudio/Src/binkace.c
@@ -810,20 +810,17 @@ static void clamp_best_quant_to_ath( U32 * best_qlevel, AUDIOFLOAT * band_ath, U
   {
     AUDIOFLOAT band_db, ath;
 
-    if ( best_qlevel[ i ] < 255 )
-    {
-      ath = band_ath[ i ] + lossy;
-
-      band_db = ( (AUDIOFLOAT) (U8) best_qlevel[ i ] ) * 0.664F;
+    ath = band_ath[ i ] + lossy;
 
-      if ( ath > band_db )
-      {
-        best_qlevel[ i ] = (U32) ( ( ath + 0.3319F ) / 0.664F );
+    band_db = ( (AUDIOFLOAT) (U8) best_qlevel[ i ] ) * 0.664F;
 
-        if ( best_qlevel[ i ] > 255 )
-          best_qlevel[ i ] = 255;
-      }
+    if ( ath > band_db )
+    {
+      best_qlevel[ i ] = (U32) ( ( ath + 0.3319F ) / 0.664F );
     }
+    
+    if (best_qlevel[i] > 95)
+        best_qlevel[i] = 95;
   }
 }
 
@@ -853,13 +850,13 @@ static void encode_one_channel( U32 transform_size,
   #else
 
   AUDIOFLOAT threshold_mult[ TOTBANDS ];
-  S16 coeffs[ MAXBUFFERSIZE * MAXCHANNELS ];
-  U16 levels[ ( ( ( MAXBUFFERSIZE * MAXCHANNELS ) + VQLENGTH ) / VQLENGTH ) ];
+  S16 coeffs[ MAXBUFFERSIZE ];
+  U16 levels[ ( ( MAXBUFFERSIZE + VQLENGTH ) / VQLENGTH ) ];
 
   #endif
 
   rrassert( TOTBANDS >= num_bands );
-  rrassert( ( MAXBUFFERSIZE * MAXCHANNELS ) >= transform_size );
+  rrassert( ( MAXBUFFERSIZE ) >= transform_size );
 
   // unquant the threshold for compression
   for ( i = 0 ; i < num_bands ; i++ )
diff --git a/Engine/Source/Runtime/BinkAudioDecoder/SDK/BinkAudio/Src/varbits.h b/Engine/Source/Runtime/BinkAudioDecoder/SDK/BinkAudio/Src/varbits.h
--- a/Engine/Source/Runtime/BinkAudioDecoder/SDK/BinkAudio/Src/varbits.h
+++ b/Engine/Source/Runtime/BinkAudioDecoder/SDK/BinkAudio/Src/varbits.h
@@ -16,7 +16,9 @@
 
 
 // variable bit macros
+//#if defined(__RADX86__)
 RADDECLAREDATA const RAD_ALIGN(U32,VarBitsLens[33],32);
+//#endif
 
 /*************
 
diff --git a/Engine/Source/Runtime/BinkAudioDecoder/SDK/BinkAudio/Src/x86_cpu.c b/Engine/Source/Runtime/BinkAudioDecoder/SDK/BinkAudio/Src/x86_cpu.c
--- a/Engine/Source/Runtime/BinkAudioDecoder/SDK/BinkAudio/Src/x86_cpu.c
+++ b/Engine/Source/Runtime/BinkAudioDecoder/SDK/BinkAudio/Src/x86_cpu.c
@@ -109,4 +109,4 @@ RADDEFFUNC void CPU_clear( void )
 #endif
 }
 
-#endif // defined(__RADX86__)
+#endif // __RADX86__
\ No newline at end of file
