Fix for bink audio 96khz regression.

diff --git a/Engine/Source/Runtime/BinkAudioDecoder/SDK/BinkAudio/Src/binka_ue_encode.cpp b/Engine/Source/Runtime/BinkAudioDecoder/SDK/BinkAudio/Src/binka_ue_encode.cpp
--- a/Engine/Source/Runtime/BinkAudioDecoder/SDK/BinkAudio/Src/binka_ue_encode.cpp
+++ b/Engine/Source/Runtime/BinkAudioDecoder/SDK/BinkAudio/Src/binka_ue_encode.cpp
@@ -544,7 +544,7 @@ uint8_t UECompressBinkAudio(
     Header.tag = 'UEBA';
     Header.PADDING = 0;
     Header.channels = (U8)WavChannels;
-    Header.rate = (U16)WavRate;
+    Header.rate = WavRate;
     Header.sample_count = SamplesPerChannel;
     Header.max_comp_space_needed = (U16)MaxBlockSize;
     Header.flags = 1;
