Use a quoted string for error message in BinkAudioInfo.cpp, to avoid invalid token caused by unpaired apostrophe (which would be parsed as a character literal).

See: https://gcc.gnu.org/onlinedocs/gcc-4.5.0/cpp/Diagnostics.html

diff --git a/Engine/Source/Runtime/BinkAudioDecoder/Module/Private/BinkAudioInfo.cpp b/Engine/Source/Runtime/BinkAudioDecoder/Module/Private/BinkAudioInfo.cpp
--- a/Engine/Source/Runtime/BinkAudioDecoder/Module/Private/BinkAudioInfo.cpp
+++ b/Engine/Source/Runtime/BinkAudioDecoder/Module/Private/BinkAudioInfo.cpp
@@ -9,7 +9,7 @@
 #include "binka_ue_decode.h"
 
 #if !defined(PLATFORM_LITTLE_ENDIAN) || !PLATFORM_LITTLE_ENDIAN
-#error Bink Audio hasn't been updated for big endian.
+#error "Bink Audio hasn't been updated for big endian."
 #endif
 
 DEFINE_LOG_CATEGORY_STATIC(LogBinkAudioDecoder, Log, All);
