Fix for some assets clicking on loop with bink audio.

diff --git a/Engine/Source/Runtime/BinkAudioDecoder/Module/Private/BinkAudioInfo.cpp b/Engine/Source/Runtime/BinkAudioDecoder/Module/Private/BinkAudioInfo.cpp
--- a/Engine/Source/Runtime/BinkAudioDecoder/Module/Private/BinkAudioInfo.cpp
+++ b/Engine/Source/Runtime/BinkAudioDecoder/Module/Private/BinkAudioInfo.cpp
@@ -96,6 +96,20 @@ FBinkAudioInfo::~FBinkAudioInfo()
 	FMemory::Free(RawMemory);
 }
 
+//-----------------------------------------------------------------------------
+//-----------------------------------------------------------------------------
+void FBinkAudioInfo::PrepareToLoop()
+{
+#if WITH_BINK_AUDIO
+	UEBinkAudioDecodeInterface* BinkInterface = UnrealBinkAudioDecodeInterface();
+	void** Streams = Decoder->Decoders();
+	for (uint8 i = 0; i < Decoder->StreamCount; i++)
+	{
+		BinkInterface->ResetStartFn(Streams[i]);
+	}
+#endif // WITH_BINK_AUDIO
+}
+
 //-----------------------------------------------------------------------------
 //-----------------------------------------------------------------------------
 void FBinkAudioInfo::SeekToTime(const float SeekTimeSeconds)
diff --git a/Engine/Source/Runtime/BinkAudioDecoder/Module/Public/BinkAudioInfo.h b/Engine/Source/Runtime/BinkAudioDecoder/Module/Public/BinkAudioInfo.h
--- a/Engine/Source/Runtime/BinkAudioDecoder/Module/Public/BinkAudioInfo.h
+++ b/Engine/Source/Runtime/BinkAudioDecoder/Module/Public/BinkAudioInfo.h
@@ -21,6 +21,7 @@ public:
 	virtual uint32 GetMaxFrameSizeSamples() const override;
 	virtual bool CreateDecoder() override;
 	virtual void SeekToTime(const float SeekToTimeSeconds) override;
+	virtual void PrepareToLoop() override;
 	virtual FDecodeResult Decode(const uint8* CompressedData, const int32 CompressedDataSize, uint8* OutPCMData, const int32 OutputPCMDataSize) override;
 	virtual bool HasError() const override;
 	//~ End IStreamedCompressedInfo Interface
