Updating bink audio encoder to support larger seek tables and lower sample rates.

diff --git a/Engine/Source/Developer/AudioFormatBink/Private/AudioFormatBink.cpp b/Engine/Source/Developer/AudioFormatBink/Private/AudioFormatBink.cpp
--- a/Engine/Source/Developer/AudioFormatBink/Private/AudioFormatBink.cpp
+++ b/Engine/Source/Developer/AudioFormatBink/Private/AudioFormatBink.cpp
@@ -78,7 +78,7 @@ public:
 		
 		void* CompressedData = 0;
 		uint32_t CompressedDataLen = 0;
-		uint8_t BinkCompressError = UECompressBinkAudio((void*)InSrcBuffer.GetData(), InSrcBuffer.Num(), InQualityInfo.SampleRate, InQualityInfo.NumChannels, CompressionLevel, 1, BinkAlloc, BinkFree, &CompressedData, &CompressedDataLen);
+		uint8_t BinkCompressError = UECompressBinkAudio((void*)InSrcBuffer.GetData(), InSrcBuffer.Num(), InQualityInfo.SampleRate, InQualityInfo.NumChannels, CompressionLevel, 1, 4096, BinkAlloc, BinkFree, &CompressedData, &CompressedDataLen);
 
 		const TCHAR* CompressErrorStr = nullptr;
 		switch (BinkCompressError)
@@ -90,6 +90,7 @@ public:
 		case BINKA_COMPRESS_ERROR_QUALITY: CompressErrorStr = TEXT("Invalid quality provided, valid is 0-9"); break;
 		case BINKA_COMPRESS_ERROR_ALLOCATORS: CompressErrorStr = TEXT("No allocators provided!"); break;
 		case BINKA_COMPRESS_ERROR_OUTPUT: CompressErrorStr = TEXT("No output pointers provided!"); break;
+		case BINKA_COMPRESS_ERROR_SEEKTABLE: CompressErrorStr = TEXT("Invalid seektable size limit specified!"); break;
 		}
 
 		if (CompressErrorStr != nullptr)
diff --git a/Engine/Source/Runtime/BinkAudioDecoder/SDK/BinkAudio/Include/binka_ue_encode.h b/Engine/Source/Runtime/BinkAudioDecoder/SDK/BinkAudio/Include/binka_ue_encode.h
--- a/Engine/Source/Runtime/BinkAudioDecoder/SDK/BinkAudio/Include/binka_ue_encode.h
+++ b/Engine/Source/Runtime/BinkAudioDecoder/SDK/BinkAudio/Include/binka_ue_encode.h
@@ -14,9 +14,10 @@ typedef void BAUECompressFreeFnType(void* Ptr);
 #define BINKA_COMPRESS_ERROR_QUALITY 4
 #define BINKA_COMPRESS_ERROR_ALLOCATORS 5
 #define BINKA_COMPRESS_ERROR_OUTPUT 6
+#define BINKA_COMPRESS_ERROR_SEEKTABLE 7
 
 #define BINKA_MAX_CHANS_STR TEXT("16") // (8*2) MAX_STREAMS stereo
-#define BINKA_MIN_RATE_STR TEXT("11 khz")
+#define BINKA_MIN_RATE_STR TEXT("2 khz")
 #define BINKA_MAX_RATE_STR TEXT("256 khz")
 
 //
@@ -30,6 +31,9 @@ typedef void BAUECompressFreeFnType(void* Ptr);
 // PcmDataLen is in bytes.
 // OutData will be filled with a buffer allocated by MemAlloc that contains the
 // compressed file data.
+// SeekTableMaxEntries is used to cap the amount of disk space used on the seek table.
+// Historically this has been capped at 4k, or 8k of data. If GenerateSeekTable is set,
+// this should be at least 1k or so to be useful.
 uint8_t UECompressBinkAudio(
     void* PcmData, 
     uint32_t PcmDataLen, 
@@ -37,6 +41,7 @@ uint8_t UECompressBinkAudio(
     uint8_t PcmChannels, 
     uint8_t Quality, 
     uint8_t GenerateSeekTable, 
+    uint16_t SeekTableMaxEntries,
     BAUECompressAllocFnType* MemAlloc,
     BAUECompressFreeFnType* MemFree,
     void** OutData, 
diff --git a/Engine/Source/Runtime/BinkAudioDecoder/SDK/BinkAudio/Src/binka_ue_encode.cpp b/Engine/Source/Runtime/BinkAudioDecoder/SDK/BinkAudio/Src/binka_ue_encode.cpp
--- a/Engine/Source/Runtime/BinkAudioDecoder/SDK/BinkAudio/Src/binka_ue_encode.cpp
+++ b/Engine/Source/Runtime/BinkAudioDecoder/SDK/BinkAudio/Src/binka_ue_encode.cpp
@@ -18,7 +18,7 @@
 #define BINKA_COMPRESS_ERROR_QUALITY 4
 #define BINKA_COMPRESS_ERROR_ALLOCATORS 5
 #define BINKA_COMPRESS_ERROR_OUTPUT 6
-
+#define BINKA_COMPRESS_ERROR_SEEKTABLE 7
 
 struct MemBufferEntry
 {
@@ -203,9 +203,8 @@ static void SeekTableBufferFree(SeekTableBuffer* seek)
 
 //-----------------------------------------------------------------------------
 //-----------------------------------------------------------------------------
-static U32 SeekTableBufferTrim(SeekTableBuffer* seek)
+static U32 SeekTableBufferTrim(SeekTableBuffer* seek, U16 max_entry_count)
 {
-    // Trim to 4096 size, or 2048 total entries.
     // Return # of frames per entry.
 
     // resolve to a single buffer for ease.
@@ -214,7 +213,7 @@ static U32 SeekTableBufferTrim(SeekTableBuffer* seek)
     memcpy((char*)table + seek->buffer.total_bytes, seek->current_stack, sizeof(U16) * seek->current_index);
 
     S32 frames_per_entry = 1;
-    while (seek->total_count > 4096)
+    while (seek->total_count > (U32)max_entry_count)
     {
         // collapse pairs.
         frames_per_entry <<= 1;
@@ -337,14 +336,14 @@ struct FPStateScope
 //-----------------------------------------------------------------------------
 uint8_t UECompressBinkAudio(
     void* WavData, uint32_t WavDataLen, uint32_t WavRate, uint8_t WavChannels, 
-    uint8_t Quality, uint8_t GenerateSeekTable, BAUECompressAllocFnType* MemAlloc, BAUECompressFreeFnType* MemFree,
+    uint8_t Quality, uint8_t GenerateSeekTable, uint16_t SeekTableMaxEntries, BAUECompressAllocFnType* MemAlloc, BAUECompressFreeFnType* MemFree,
     void** OutData, uint32_t* OutDataLen)
 {
     FPStateScope FixFloatingPoint;
 
     if (WavChannels == 0)
         return BINKA_COMPRESS_ERROR_CHANS;
-    if (WavRate > 256000 || WavRate < 11025)
+    if (WavRate > 256000 || WavRate < 2000)
         return BINKA_COMPRESS_ERROR_RATE;
     if (Quality > 9)
         return BINKA_COMPRESS_ERROR_QUALITY;
@@ -354,6 +353,8 @@ uint8_t UECompressBinkAudio(
     if (OutData == nullptr ||
         OutDataLen == nullptr)
         return BINKA_COMPRESS_ERROR_OUTPUT;
+    if (GenerateSeekTable && SeekTableMaxEntries < 2)
+        return BINKA_COMPRESS_ERROR_SEEKTABLE;
 
     //
     // Deinterlace the input.
@@ -537,8 +538,8 @@ uint8_t UECompressBinkAudio(
             break;
     }
 
-    // Trim the table to 4k and get how many frames ended up per entry.
-    U32 FramesPerEntry = SeekTableBufferTrim(&SeekTable);
+    // Trim the table and get how many frames ended up per entry.
+    U32 FramesPerEntry = SeekTableBufferTrim(&SeekTable, SeekTableMaxEntries);
 
     BinkAudioFileHeader Header;
     Header.tag = 'UEBA';
@@ -581,7 +582,8 @@ struct fuzzer_input
     uint32_t Rate;
     uint8_t Chans;
     uint8_t Quality;
-    uint8_t GenTable;    
+    uint8_t GenTable;
+    uint16_t GenTableMaxEntries;  
 };
 
 static void* AllocThunk(uintptr_t ByteCount) { return malloc(ByteCount); }
@@ -600,7 +602,7 @@ extern "C" int LLVMFuzzerTestOneInput(const uint8_t *Data, size_t Size)
     uint32_t CompressedLen;
 
     UECompressBinkAudio((void*)(Data + sizeof(fuzzer_input)), (uint32_t)Size - sizeof(fuzzer_input),
-        input->Rate, input->Chans, input->Quality, input->GenTable, AllocThunk, FreeThunk,
+        input->Rate, input->Chans, input->Quality, input->GenTable, input->GenTableMaxEntries, AllocThunk, FreeThunk,
         &CompressedData, &CompressedLen);
 
     FreeThunk(CompressedData);
diff --git a/Engine/Source/Runtime/BinkAudioDecoder/SDK/BinkAudio/Src/binka_ue_encode_test.cpp b/Engine/Source/Runtime/BinkAudioDecoder/SDK/BinkAudio/Src/binka_ue_encode_test.cpp
--- a/Engine/Source/Runtime/BinkAudioDecoder/SDK/BinkAudio/Src/binka_ue_encode_test.cpp
+++ b/Engine/Source/Runtime/BinkAudioDecoder/SDK/BinkAudio/Src/binka_ue_encode_test.cpp
@@ -20,7 +20,7 @@ int main(int argc, char** argv)
 
     void* Compressed;
     uint32_t CompressedLen;
-    UECompressBinkAudio(data, 48000*2, 48000, 1, 0, 1, allocator_thunk, free_thunk, &Compressed, &CompressedLen);
+    UECompressBinkAudio(data, 48000*2, 48000, 1, 0, 1, 4096, allocator_thunk, free_thunk, &Compressed, &CompressedLen);
 
     printf("Compressed %d %d\n", 48000*2, CompressedLen);
 
