Fix for very large binks files creating garbage data

diff --git a/Engine/Source/Developer/AudioFormatBink/Private/AudioFormatBink.cpp b/Engine/Source/Developer/AudioFormatBink/Private/AudioFormatBink.cpp
--- a/Engine/Source/Developer/AudioFormatBink/Private/AudioFormatBink.cpp
+++ b/Engine/Source/Developer/AudioFormatBink/Private/AudioFormatBink.cpp
@@ -56,7 +56,9 @@ namespace AudioFormatBinkPrivate
 		static const uint16 BinkDefault = 4096;
 		static const uint16 BinkMax		= TNumericLimits<uint16>::Max();
 
-		return IntCastChecked<uint16>(FMath::Clamp<uint64>(MaxEntries, BinkDefault, BinkMax));
+		const uint64 Clamped = FMath::Clamp<uint64>(MaxEntries, BinkDefault, BinkMax);
+			
+		return IntCastChecked<uint16>(Clamped);
 	}
 }
 
@@ -68,7 +70,7 @@ class FAudioFormatBink : public IAudioFormat
 	enum
 	{
 		/** Version for Bink Audio format, this becomes part of the DDC key. */
-		UE_AUDIO_BINK_VER = 9,
+		UE_AUDIO_BINK_VER = 11,
 	};
 
 public:
@@ -122,6 +124,7 @@ public:
 		case BINKA_COMPRESS_ERROR_ALLOCATORS: CompressErrorStr = TEXT("No allocators provided!"); break;
 		case BINKA_COMPRESS_ERROR_OUTPUT: CompressErrorStr = TEXT("No output pointers provided!"); break;
 		case BINKA_COMPRESS_ERROR_SEEKTABLE: CompressErrorStr = TEXT("Invalid seektable size limit specified!"); break;
+		case BINKA_COMPRESS_ERROR_SIZE: CompressErrorStr = TEXT("Input file too big - can't fit offsets in seek table!"); break;
 		}
 
 		if (CompressErrorStr != nullptr)
@@ -309,11 +312,11 @@ public:
 		static const uint32 ActualAudioOffset = sizeof(BinkAudioFileHeader) + 0; // No entries as we're stripping it.
 
 		// Decode and copy out the seek-table. (it's stored as deltas).
-		uint16* EncodedSeekTable = (uint16*)(InOutBuffer.GetData() + sizeof(BinkAudioFileHeader));
+		const uint16* EncodedSeekTable = (uint16*)(InOutBuffer.GetData() + sizeof(BinkAudioFileHeader));
 		uint32 CurrentSeekOffset = ActualAudioOffset;
 		uint32 CurrentTimeOffset = 0;
 		
-		int32 SamplesPerEntry = Header->blocks_per_seek_table_entry * AudioFormatBinkPrivate::GetMaxFrameSizeSamples(Header->rate);
+		const int32 SamplesPerEntry = Header->blocks_per_seek_table_entry * AudioFormatBinkPrivate::GetMaxFrameSizeSamples(Header->rate);
 		OutSeektable.Offsets.SetNum(Header->seek_table_entry_count);
 		OutSeektable.Times.SetNum(Header->seek_table_entry_count);
 
@@ -325,6 +328,13 @@ public:
 			CurrentTimeOffset += SamplesPerEntry;
 		}
 
+		// Check that the last block which spans the last offset in the table and end of file
+		// is a reasonable size.
+		if (!ensure(InOutBuffer.Num() - OutSeektable.Offsets.Last() < 1024*1024))
+		{
+			return false;
+		}
+
 		// Strip the seek-table from the buffer now we've copied it.
 		StripSeek(InOutBuffer);
 
@@ -355,4 +365,4 @@ public:
 	}
 };
 
-IMPLEMENT_MODULE( FAudioPlatformBinkModule, AudioFormatBink);
\ No newline at end of file
+IMPLEMENT_MODULE( FAudioPlatformBinkModule, AudioFormatBink);
diff --git a/Engine/Source/Runtime/BinkAudioDecoder/SDK/BinkAudio/Include/binka_ue_encode.h b/Engine/Source/Runtime/BinkAudioDecoder/SDK/BinkAudio/Include/binka_ue_encode.h
--- a/Engine/Source/Runtime/BinkAudioDecoder/SDK/BinkAudio/Include/binka_ue_encode.h
+++ b/Engine/Source/Runtime/BinkAudioDecoder/SDK/BinkAudio/Include/binka_ue_encode.h
@@ -5,8 +5,6 @@
 typedef void* BAUECompressAllocFnType(uintptr_t ByteCount);
 typedef void BAUECompressFreeFnType(void* Ptr);
 
-// (NOTE! This was copied after integrating the latest source, so when updating make sure
-// it stays around!
 #define BINKA_COMPRESS_SUCCESS 0
 #define BINKA_COMPRESS_ERROR_CHANS 1
 #define BINKA_COMPRESS_ERROR_SAMPLES 2
@@ -15,6 +13,7 @@ typedef void BAUECompressFreeFnType(void* Ptr);
 #define BINKA_COMPRESS_ERROR_ALLOCATORS 5
 #define BINKA_COMPRESS_ERROR_OUTPUT 6
 #define BINKA_COMPRESS_ERROR_SEEKTABLE 7
+#define BINKA_COMPRESS_ERROR_SIZE 8
 
 #define BINKA_MAX_CHANS_STR TEXT("16") // (8*2) MAX_STREAMS stereo
 #define BINKA_MIN_RATE_STR TEXT("2 khz")
diff --git a/Engine/Source/Runtime/BinkAudioDecoder/SDK/BinkAudio/Src/binka_ue_encode.cpp b/Engine/Source/Runtime/BinkAudioDecoder/SDK/BinkAudio/Src/binka_ue_encode.cpp
--- a/Engine/Source/Runtime/BinkAudioDecoder/SDK/BinkAudio/Src/binka_ue_encode.cpp
+++ b/Engine/Source/Runtime/BinkAudioDecoder/SDK/BinkAudio/Src/binka_ue_encode.cpp
@@ -11,15 +11,6 @@
 
 #define MAX_STREAMS 8
 
-#define BINKA_COMPRESS_SUCCESS 0
-#define BINKA_COMPRESS_ERROR_CHANS 1
-#define BINKA_COMPRESS_ERROR_SAMPLES 2
-#define BINKA_COMPRESS_ERROR_RATE 3
-#define BINKA_COMPRESS_ERROR_QUALITY 4
-#define BINKA_COMPRESS_ERROR_ALLOCATORS 5
-#define BINKA_COMPRESS_ERROR_OUTPUT 6
-#define BINKA_COMPRESS_ERROR_SEEKTABLE 7
-
 struct MemBufferEntry
 {
     U32 bytecount;
@@ -212,30 +203,35 @@ static U32 SeekTableBufferTrim(SeekTableBuffer* seek, U16 max_entry_count)
     MemBufferWriteBuffer(&seek->buffer, table);
     memcpy((char*)table + seek->buffer.total_bytes, seek->current_stack, sizeof(U16) * seek->current_index);
 
-    S32 frames_per_entry = 1;
+    U32 frames_per_entry = 1;
     while (seek->total_count > (U32)max_entry_count)
     {
         // collapse pairs.
         frames_per_entry <<= 1;
 
-        S32 new_total_count = seek->total_count;
         U32 read_index = 0;
         U32 write_index = 0;
         while (read_index < seek->total_count)
         {
-            U16 total = table[read_index];
+            U32 total = table[read_index];
             if (read_index + 1 < seek->total_count)
             {
                 total += table[read_index + 1];
-                new_total_count--;
             }
 
-            table[write_index] = total;
+            if (total > 65535)
+            {
+                // we can't fit the offset in the table - fail to trim, this file is too big
+                // with the given max seek table size.
+                return ~0U;
+            }
+
+            table[write_index] = (U16)total;
             write_index++;
             read_index += 2;
         }
 
-        seek->total_count = new_total_count;
+        seek->total_count = write_index;
     }
 
     seek->collapsed = table;
@@ -538,27 +534,40 @@ uint8_t UECompressBinkAudio(
             break;
     }
 
+    uint8_t Result = BINKA_COMPRESS_SUCCESS;
     // Trim the table and get how many frames ended up per entry.
     U32 FramesPerEntry = SeekTableBufferTrim(&SeekTable, SeekTableMaxEntries);
-
-    BinkAudioFileHeader Header;
-    Header.tag = 'UEBA';
-    Header.PADDING = 0;
-    Header.channels = (U8)WavChannels;
-    Header.rate = WavRate;
-    Header.sample_count = SamplesPerChannel;
-    Header.max_comp_space_needed = (U16)MaxBlockSize;
-    Header.flags = 1;
-    Header.version = 1;
-    Header.output_file_size = DataBuffer.total_bytes + sizeof(BinkAudioFileHeader) + SeekTable.total_count*sizeof(U16);
-    Header.blocks_per_seek_table_entry = (U16)FramesPerEntry;
-    Header.seek_table_entry_count = (U16)SeekTable.total_count;
-
-    // Create the actual output buffer.
-    char* Output = (char*)MemAlloc(Header.output_file_size);
-    memcpy(Output, &Header, sizeof(BinkAudioFileHeader));
-    memcpy(Output + sizeof(BinkAudioFileHeader), SeekTable.collapsed, SeekTable.total_count * sizeof(U16));
-    MemBufferWriteBuffer(&DataBuffer, Output + sizeof(BinkAudioFileHeader) + SeekTable.total_count * sizeof(U16));
+    if (FramesPerEntry == ~0U)
+    {
+        // The file is too big - can't fit offsets in to U16.
+        Result = BINKA_COMPRESS_ERROR_SIZE;
+    }
+    else
+    {
+        BinkAudioFileHeader Header;
+        Header.tag = 'UEBA';
+        Header.PADDING = 0;
+        Header.channels = (U8)WavChannels;
+        Header.rate = WavRate;
+        Header.sample_count = SamplesPerChannel;
+        Header.max_comp_space_needed = (U16)MaxBlockSize;
+        Header.flags = 1;
+        Header.version = 1;
+        Header.output_file_size = DataBuffer.total_bytes + sizeof(BinkAudioFileHeader) + SeekTable.total_count*sizeof(U16);
+        Header.blocks_per_seek_table_entry = (U16)FramesPerEntry;
+        Header.seek_table_entry_count = (U16)SeekTable.total_count;
+
+        // Create the actual output buffer.
+        char* Output = (char*)MemAlloc(Header.output_file_size);
+        memcpy(Output, &Header, sizeof(BinkAudioFileHeader));
+        memcpy(Output + sizeof(BinkAudioFileHeader), SeekTable.collapsed, SeekTable.total_count * sizeof(U16));
+        MemBufferWriteBuffer(&DataBuffer, Output + sizeof(BinkAudioFileHeader) + SeekTable.total_count * sizeof(U16));
+
+        *OutData = Output;
+        *OutDataLen = Header.output_file_size;
+
+        Result = BINKA_COMPRESS_SUCCESS;
+    }
 
     MemBufferFree(&DataBuffer);
     SeekTableBufferFree(&SeekTable);
@@ -568,10 +577,8 @@ uint8_t UECompressBinkAudio(
         BinkAudioCompressClose(hBink[i]);
     }
 
-    *OutData = Output;
-    *OutDataLen = Header.output_file_size;
 
-    return BINKA_COMPRESS_SUCCESS;
+    return Result;
 }
 
 #if 0
@@ -610,3 +617,18 @@ extern "C" int LLVMFuzzerTestOneInput(const uint8_t *Data, size_t Size)
     return 0;
 }
 #endif
+
+/* @cdep pre
+    $NoAutoInclude
+    $Defaults
+    $requires(popmal.c)
+    $requires(binkace.c)
+    $requires(binka_ue_encode.cpp)
+    $requires(radfft.cpp)
+    $requires(varbits.c)
+    $requires(ranged_log.cpp)
+*/
+
+/* @cdep post
+    $Build
+*/
diff --git a/Engine/Source/Runtime/Engine/Private/AudioDerivedData.cpp b/Engine/Source/Runtime/Engine/Private/AudioDerivedData.cpp
--- a/Engine/Source/Runtime/Engine/Private/AudioDerivedData.cpp
+++ b/Engine/Source/Runtime/Engine/Private/AudioDerivedData.cpp
@@ -541,32 +541,46 @@ class FStreamedAudioCacheDerivedDataWorker : public FNonAbandonableTask
 		while (Current < SourceEnd)
 		{	
 			// A block is the span between two neighboring offsets in our table.
-			uint32 BlockStart = Offsets[CurrentBlock];
-			uint32 BlockEnd = CurrentBlock < Offsets.Num() - 1 ? Offsets[CurrentBlock + 1] : SourceLen;
-			uint32 BlockSize = BlockEnd - BlockStart;
+			const uint32 BlockStart = Offsets[CurrentBlock];
+			const uint32 BlockEnd = CurrentBlock < Offsets.Num() - 1 ? Offsets[CurrentBlock + 1] : SourceLen;
+			if (!ensureMsgf(BlockStart < BlockEnd, TEXT("Malformed Table: BlockStart=%u, BlockEnd=%u"), BlockStart, BlockEnd))
+			{
+				return false;
+			}
+			const uint32 BlockSize = BlockEnd - BlockStart;
+			if (!ensureMsgf(BlockSize < 1024*1024, TEXT("Malformed Table: Block is too large to be real. BlockSize=%u"), BlockSize) )
+			{
+				return false;
+			}
 			
 			// Current chunk stats.
-			int32 CurrentTableSize = FStreamedAudioChunkSeekTable::CalcSize(BlockCount, Mode);
-			uint32 CurrentChunkSize = (Current - ChunkStart) + CurrentTableSize;
+			const int32 CurrentTableSize = FStreamedAudioChunkSeekTable::CalcSize(BlockCount, Mode);
+			checkf(CurrentTableSize >= 0, TEXT("CurrentTableSize=%d"), CurrentTableSize);
+			const uint32 CurrentChunkSize = (Current - ChunkStart) + CurrentTableSize;
 			
 			// Adding a new item will grow the table by 1.
-			int32 NewTableSize = FStreamedAudioChunkSeekTable::CalcSize(BlockCount + 1, Mode);
-			int32 TableDelta = NewTableSize - CurrentTableSize;
+			const int32 NewTableSize = FStreamedAudioChunkSeekTable::CalcSize(BlockCount + 1, Mode);
+			checkf(NewTableSize >= 0, TEXT("NewTableSize=%d"), NewTableSize);
+			const int32 TableDelta = NewTableSize - CurrentTableSize;
+			checkf(TableDelta > 0, TEXT("Table should grow when we add a new entry: TableDelta=%d"), TableDelta);
 			
 			// Accumulated enough?
 			if (CurrentChunkSize + BlockSize + TableDelta > Budget)
 			{
 				// can't add this chunk, emit.
 				TArray<uint8> Chunk(ChunkStart, Current - ChunkStart);
-				int32 AudioSize = Chunk.Num();
+				const int32 AudioSize = Chunk.Num();
 
 				CreateChunkSeektableAndPrefix(Mode, InTable, ChunkStart-Source, Current-Source, Chunk, OutChunkOffsets, CurrentTableSize, BlockCount, AudioSize, Budget);
 		
+				UE_LOG(LogAudio, Verbose, TEXT("Adding Chunk %d: Blocks=%d (%u bytes), SeekTableEntries=%d (%d bytes), ChunkSize=%u bytes, PercentFull=%2.2f, Remaining=%u bytes"),
+					OutBuffers.Num(), BlockCount, AudioSize, BlockCount, CurrentTableSize, Chunk.Num(), ((float)Chunk.Num() / Budget) * 100.f, Budget-Chunk.Num());
+
 				OutBuffers.Add(Chunk);
-				
+
 				ChunkStart = Current;
 				Budget = InMaxChunkSize;
-				BlockCount = 0;
+				BlockCount = 0;				
 			
 				// retry.
 				continue;
@@ -577,7 +591,7 @@ class FStreamedAudioCacheDerivedDataWorker : public FNonAbandonableTask
 			BlockCount++;
 			CurrentBlock++;
 
-			UE_LOG(LogAudio, Verbose, TEXT("Adding Block: Chunk_Block_Count=%d, BlockSize=%d, ChunkSize=%d"), BlockCount, BlockSize, Current - ChunkStart);
+			UE_LOG(LogAudio, VeryVerbose, TEXT("\tAdding Block: Chunk_Block_Count=%d, BlockSize=%d, ChunkSize=%lld"), BlockCount, BlockSize, Current - ChunkStart);
 		}
 
 		// emit any remainder chunks
@@ -585,9 +599,13 @@ class FStreamedAudioCacheDerivedDataWorker : public FNonAbandonableTask
 		{
 			// emit this chunk
 			TArray<uint8> Chunk(ChunkStart, Current - ChunkStart);
-			int32 AudioSize = Chunk.Num();
-			int32 CurrentTableSize = FStreamedAudioChunkSeekTable::CalcSize(BlockCount, Mode);
+			const int32 AudioSize = Chunk.Num();
+			const int32 CurrentTableSize = FStreamedAudioChunkSeekTable::CalcSize(BlockCount, Mode);
 			CreateChunkSeektableAndPrefix(Mode, InTable, ChunkStart - Source, Current - Source, Chunk, OutChunkOffsets, CurrentTableSize, BlockCount, AudioSize, Budget);
+
+			UE_LOG(LogAudio, Verbose, TEXT("Adding Chunk %d: Blocks=%d (%u bytes), SeekTableEntries=%d (%d bytes), ChunkSize=%u bytes, PercentFull=%2.2f, Remaining=%u bytes"),
+				OutBuffers.Num(), BlockCount, AudioSize, BlockCount, CurrentTableSize, Chunk.Num(), ((float)Chunk.Num() / Budget) * 100.f, Budget - Chunk.Num());
+
 			OutBuffers.Add(Chunk);
 		}
 
