Add an arm64 iOS Simulator build of the BinkAudioDecoder library

#rnx

diff --git a/Engine/Source/Runtime/BinkAudioDecoder/Module/BinkAudioDecoder.Build.cs b/Engine/Source/Runtime/BinkAudioDecoder/Module/BinkAudioDecoder.Build.cs
--- a/Engine/Source/Runtime/BinkAudioDecoder/Module/BinkAudioDecoder.Build.cs
+++ b/Engine/Source/Runtime/BinkAudioDecoder/Module/BinkAudioDecoder.Build.cs
@@ -24,7 +24,8 @@ public class BinkAudioDecoder : ModuleRules
         }
         if (Target.Platform == UnrealTargetPlatform.IOS)
         {
-            return Path.Combine(ModuleDirectory, "..", "SDK", "BinkAudio", "Lib", "libbinka_ue_decode_ios_static.a");
+			string LibExt = (Target.Architecture == UnrealArch.IOSSimulator) ? "_sim.a" : ".a";
+            return Path.Combine(ModuleDirectory, "..", "SDK", "BinkAudio", "Lib", "libbinka_ue_decode_ios_static" + LibExt);
         }
         if (Target.Platform == UnrealTargetPlatform.TVOS)
         {
diff --git a/Engine/Source/Runtime/BinkAudioDecoder/SDK/BinkAudio/BuildDecoderIos.command b/Engine/Source/Runtime/BinkAudioDecoder/SDK/BinkAudio/BuildDecoderIos.command
--- a/Engine/Source/Runtime/BinkAudioDecoder/SDK/BinkAudio/BuildDecoderIos.command
+++ b/Engine/Source/Runtime/BinkAudioDecoder/SDK/BinkAudio/BuildDecoderIos.command
@@ -20,6 +20,7 @@ mkdir $BUILD_LOCATION/ios/arm64
 mkdir $BUILD_LOCATION/ios/arm7
 mkdir $BUILD_LOCATION/ios/arm7s
 mkdir $BUILD_LOCATION/ios/x64
+mkdir $BUILD_LOCATION/ios/sim
 
 SIM_ARGS=(
 	-arch x86_64
@@ -28,6 +29,12 @@ SIM_ARGS=(
 	-isysroot /Applications/Xcode.app/Contents/Developer/Platforms/iPhoneSimulator.platform/Developer/SDKs/iPhoneSimulator.sdk
 )
 
+SIM_ARGS_ARM64=(
+	-arch arm64
+	-miphonesimulator-version-min=15
+	-isysroot /Applications/Xcode.app/Contents/Developer/Platforms/iPhoneSimulator.platform/Developer/SDKs/iPhoneSimulator.sdk
+)
+
 ARM64_ARGS=(
 	-arch arm64
 	--target=arm-apple-ios10
@@ -98,6 +105,14 @@ done
 ar rcs $OUTPUT_LOCATION/libbinka_ue_decode_ios_arm7s_static.a ${OUTPUT[@]}
 OUTPUT=()
 
+for source_file in "${SOURCES[@]}"
+do
+	clang ${SIM_ARGS_ARM64[@]} ${COMMON_ARGS[@]} -o $BUILD_LOCATION/ios/sim/${source_file%.*}.o $SOURCE_LOCATION/$source_file
+	OUTPUT+=( $BUILD_LOCATION/ios/sim/${source_file%.*}.o )
+done
+ar rcs $OUTPUT_LOCATION/libbinka_ue_decode_ios_static_sim.a ${OUTPUT[@]}
+OUTPUT=()
+
 SOURCES+=(x86_cpu.c)
 for source_file in "${SOURCES[@]}"
 do
