BinkAudioInfo should "Reset" its decoders every time it seeks, not just when "PrepareToLoop" is called.

#tests PIE

diff --git a/Engine/Source/Runtime/BinkAudioDecoder/Module/Private/BinkAudioInfo.cpp b/Engine/Source/Runtime/BinkAudioDecoder/Module/Private/BinkAudioInfo.cpp
--- a/Engine/Source/Runtime/BinkAudioDecoder/Module/Private/BinkAudioInfo.cpp
+++ b/Engine/Source/Runtime/BinkAudioDecoder/Module/Private/BinkAudioInfo.cpp
@@ -98,7 +98,7 @@ FBinkAudioInfo::~FBinkAudioInfo()
 
 //-----------------------------------------------------------------------------
 //-----------------------------------------------------------------------------
-void FBinkAudioInfo::PrepareToLoop()
+void FBinkAudioInfo::NotifySeek()
 {
 #if WITH_BINK_AUDIO
 	UEBinkAudioDecodeInterface* BinkInterface = UnrealBinkAudioDecodeInterface();
@@ -113,7 +113,8 @@ void FBinkAudioInfo::PrepareToLoop()
 //-----------------------------------------------------------------------------
 //-----------------------------------------------------------------------------
 void FBinkAudioInfo::SeekToTime(const float SeekTimeSeconds)
-{	
+{
+	NotifySeek();
 	// If there's no seek table on the header, fall-back to Super implementation.
 	if (Decoder->SeekTableCount == 0)
 	{
@@ -121,6 +122,8 @@ void FBinkAudioInfo::SeekToTime(const float SeekTimeSeconds)
 		return;
 	}
 
+	// no need to reset decoder here. Called in "SeekToFrame"
+
 	// convert seconds to frames and call SeekToFrame
 	uint32 SeekTimeFrames = 0;
 	if (SeekTimeSeconds > 0)
@@ -133,13 +136,14 @@ void FBinkAudioInfo::SeekToTime(const float SeekTimeSeconds)
 
 void FBinkAudioInfo::SeekToFrame(const uint32 InFrameNum)
 {
+	NotifySeek();
 	// If there's no seek table on the header, fall-back to Super implementation.
 	if (Decoder->SeekTableCount == 0)
 	{
 		Super::SeekToFrame(InFrameNum);
 		return;
 	}
-
+	
 	uint32 SeekTimeFrames = InFrameNum;
 	uint32 SeekTimeSamples = SeekTimeFrames * NumChannels;
 
diff --git a/Engine/Source/Runtime/BinkAudioDecoder/Module/Public/BinkAudioInfo.h b/Engine/Source/Runtime/BinkAudioDecoder/Module/Public/BinkAudioInfo.h
--- a/Engine/Source/Runtime/BinkAudioDecoder/Module/Public/BinkAudioInfo.h
+++ b/Engine/Source/Runtime/BinkAudioDecoder/Module/Public/BinkAudioInfo.h
@@ -22,14 +22,15 @@ public:
 	virtual bool CreateDecoder() override;
 	virtual void SeekToTime(const float SeekToTimeSeconds) override;
 	virtual void SeekToFrame(const uint32 SeekTimeFrames) override;
-	virtual void PrepareToLoop() override;
 	virtual FDecodeResult Decode(const uint8* CompressedData, const int32 CompressedDataSize, uint8* OutPCMData, const int32 OutputPCMDataSize) override;
 	virtual bool HasError() const override;
 	//~ End IStreamedCompressedInfo Interface
 
 protected:
 	using Super = IStreamedCompressedInfo;
-
+	
+	void NotifySeek();
+	
 	// copied from header during ParseHeader
 	uint32 MaxCompSpaceNeeded;
 	uint32 SampleRate;
