Cleanup pass on Bink Audio decoder - avoid creating buffers since the common case can decode directly to the output without any intermediaries. For a typical stereo sound, this saves around 15kb of memory per sound. ~30kb for 7.1.

diff --git a/Engine/Source/Runtime/BinkAudioDecoder/Module/Private/BinkAudioInfo.cpp b/Engine/Source/Runtime/BinkAudioDecoder/Module/Private/BinkAudioInfo.cpp
--- a/Engine/Source/Runtime/BinkAudioDecoder/Module/Private/BinkAudioInfo.cpp
+++ b/Engine/Source/Runtime/BinkAudioDecoder/Module/Private/BinkAudioInfo.cpp
@@ -42,14 +42,12 @@ struct BinkAudioDecoder
 {
 	// # of low level bink audio streams (max 2 chans each)
 	uint8 StreamCount;
-	// # of current bytes in the res
-	uint32 OutputReservoirValidBytes;
+	// where to copy the next output from
+	uint32 OutputReservoirReadOffset;
 	// # max size of the res.
 	uint32 OutputReservoirTotalBytes;
-	// offsets to the various pieces of the decoder
-	uint16 ToDeinterlaceBufferOffset32;
+	// offset to the seek table divided by 32.
 	uint16 ToSeekTableOffset32;
-	uint16 ToOutputReservoirOffset32;
 	// # of entries in the seek table - the entries are byte sizes, so 
 	// when decoding the seek table we convert to file offsets. This means
 	// we actually have +1 entries in the decoded seek table due to the standard
@@ -69,14 +67,6 @@ struct BinkAudioDecoder
 	{
 		return (void**)PTR_ADD(this, Align32(sizeof(uint8) * StreamCount + sizeof(BinkAudioDecoder)));
 	}
-	uint8* OutputReservoir()
-	{
-		return (uint8*)(PTR_ADD(this, ToOutputReservoirOffset32 * 32U));
-	}
-	int16* DeinterlaceBuffer()
-	{
-		return (int16*)(PTR_ADD(this, ToDeinterlaceBufferOffset32 * 32U));
-	}
 	uint32* SeekTable() // [SeekTableCount + 1] if SeekTableCount != 0, otherwise invalid.
 	{
 		return (uint32*)(PTR_ADD(this, ToSeekTableOffset32 * 32U));
@@ -93,7 +83,7 @@ FBinkAudioInfo::FBinkAudioInfo()
 //-----------------------------------------------------------------------------
 FBinkAudioInfo::~FBinkAudioInfo()
 {
-	FMemory::Free(RawMemory);
+
 }
 
 //-----------------------------------------------------------------------------
@@ -162,7 +152,8 @@ void FBinkAudioInfo::SeekToFrame(const uint32 InFrameNum)
 	uint32 OffsetToBlock = Decoder->SeekTable()[SeekTableIndex] + sizeof(BinkAudioFileHeader) + Decoder->SeekTableCount * sizeof(uint16);
 
 	Decoder->ConsumeFrameCount = SeekTableOffset;
-	Decoder->OutputReservoirValidBytes = 0;
+	Decoder->OutputReservoirTotalBytes = 0;
+	Decoder->OutputReservoirReadOffset = 0;
 
 #if WITH_BINK_AUDIO
 	UEBinkAudioDecodeInterface* BinkInterface = UnrealBinkAudioDecodeInterface();
@@ -291,17 +282,6 @@ bool FBinkAudioInfo::CreateDecoder()
 
 	// Figure memory for buffers:
 
-	// Deinterlace - buffer space for interleaving multi stream binks in to a standard
-	// interleaved format.
-	uint32 DeinterlaceMemory = 0;
-	if (StreamCount > 1)
-	{
-		DeinterlaceMemory = GetMaxFrameSizeSamples() * sizeof(int16) * 2;
-	}
-
-	// Space to store a decoded block.
-	uint32 OutputReservoirMemory = NumChannels * GetMaxFrameSizeSamples() * sizeof(int16);
-
 	// Space for the decoder state
 	uint32 DecoderMemoryTotal = 0;	
 	uint32 DecoderMemoryPerStream[MAX_BINK_AUDIO_CHANNELS / 2];
@@ -333,18 +313,15 @@ bool FBinkAudioInfo::CreateDecoder()
 		SeekTableMemory = Align32(sizeof(uint32) * (Header->seek_table_entry_count + 1));
 	}
 
-	uint32 TotalMemory = DecoderMemoryTotal + PtrMemory + StructMemory + OutputReservoirMemory + DeinterlaceMemory + SeekTableMemory;
+	uint32 TotalMemory = DecoderMemoryTotal + PtrMemory + StructMemory + SeekTableMemory;
 
 	//
 	// Allocate and save offsets
 	//
-	RawMemory = (uint8*)FMemory::Malloc(TotalMemory, 16);
-	memset(RawMemory, 0, TotalMemory);
-	
-	Decoder = (BinkAudioDecoder*)RawMemory;
+	RawMemory.AddZeroed(TotalMemory);	
+	Decoder = (BinkAudioDecoder*)RawMemory.GetData();
 
 	Decoder->StreamCount = StreamCount;
-	Decoder->OutputReservoirTotalBytes = OutputReservoirMemory;
 	Decoder->SeekTableCount = Header->seek_table_entry_count;
 	Decoder->FramesPerSeekTableEntry = Header->blocks_per_seek_table_entry;
 
@@ -372,13 +349,7 @@ bool FBinkAudioInfo::CreateDecoder()
 		}
 	}
 
-	Decoder->ToOutputReservoirOffset32 = (uint16)((CurrentMemory - RawMemory) / 32);
-	CurrentMemory += OutputReservoirMemory;
-
-	Decoder->ToDeinterlaceBufferOffset32 = (uint16)((CurrentMemory - RawMemory) / 32);
-	CurrentMemory += DeinterlaceMemory;
-
-	Decoder->ToSeekTableOffset32 = (uint16)((CurrentMemory - RawMemory) / 32);
+	Decoder->ToSeekTableOffset32 = (uint16)((CurrentMemory - RawMemory.GetData()) / 32);
 	CurrentMemory += SeekTableMemory;
 
 	// Decode the seek table
@@ -440,17 +411,22 @@ uint32 FBinkAudioInfo::GetMaxFrameSizeSamples() const
 //-----------------------------------------------------------------------------
 FDecodeResult FBinkAudioInfo::Decode(const uint8* CompressedData, const int32 CompressedDataSize, uint8* OutPCMData, const int32 OutputPCMDataSize)
 {
-	UEBinkAudioDecodeInterface* BinkInterface = 0;
+	UEBinkAudioDecodeInterface* BinkInterface = nullptr;
 
 #if WITH_BINK_AUDIO
 	BinkInterface = UnrealBinkAudioDecodeInterface();
 #endif // WITH_BINK_AUDIO
 
-	if (BinkInterface == 0)
+	if (BinkInterface == nullptr)
 	{
 		return FDecodeResult(); // only happens with no libs.
 	}
 
+	// \todo consider a system wide shared buffer to prevent allocations and avoid holding memory needlessly.
+	// ... also this is at most 7680 bytes, might be OK to just throw on the stack?
+	TArray<int16, TAlignedHeapAllocator<16>> DeinterleavedDecodeBuffer;
+
+
 	uint32 RemnOutputPCMDataSize = OutputPCMDataSize;
 	uint32 RemnCompressedDataSize = CompressedDataSize;
 	const uint8* CompressedDataEnd = CompressedData + CompressedDataSize;
@@ -462,30 +438,25 @@ FDecodeResult FBinkAudioInfo::Decode(const uint8* CompressedData, const int32 Co
 	// (+8 for max block header size)
 	uint8* StackBlockBuffer = (uint8*)alloca(MaxCompSpaceNeeded + BINK_UE_DECODER_END_INPUT_SPACE + 8);
 
-	const uint32 DecodeSize = GetMaxFrameSizeSamples() * sizeof(int16) * NumChannels;
+	const uint32 DecodeFrames = GetMaxFrameSizeSamples();
+	const uint32 FrameSize = NumChannels * sizeof(int16);
+	const uint32 AllStreamsDecodeSize = DecodeFrames * FrameSize;
 	while (RemnOutputPCMDataSize)
 	{
 		//
 		// Drain the output reservoir before attempting a decode.
 		//
-		if (Decoder->OutputReservoirValidBytes)
+		if (Decoder->OutputReservoirTotalBytes > Decoder->OutputReservoirReadOffset)
 		{
-			uint32 CopyBytes = Decoder->OutputReservoirValidBytes;
+			uint32 CopyBytes = Decoder->OutputReservoirTotalBytes - Decoder->OutputReservoirReadOffset;
 			if (CopyBytes > RemnOutputPCMDataSize)
 			{
 				CopyBytes = RemnOutputPCMDataSize;
 			}
 
-			FMemory::Memcpy(OutPCMData, Decoder->OutputReservoir(), CopyBytes);
+			FMemory::Memcpy(OutPCMData, OutputReservoir.GetData() + Decoder->OutputReservoirReadOffset, CopyBytes);
 
-			// We use memmove here because we expect it to be very rare that we don't
-			// consume the entire output reservoir in a call, so it's not worth managing
-			// a cursor. Just move down the remnants, which we expect to be zero.
-			if (Decoder->OutputReservoirValidBytes != CopyBytes)
-			{
-				FMemory::Memmove(Decoder->OutputReservoir(), Decoder->OutputReservoir() + CopyBytes, Decoder->OutputReservoirValidBytes - CopyBytes);
-			}
-			Decoder->OutputReservoirValidBytes -= CopyBytes;
+			Decoder->OutputReservoirReadOffset += CopyBytes;
 
 			RemnOutputPCMDataSize -= CopyBytes;
 			OutPCMData += CopyBytes;
@@ -499,7 +470,7 @@ FDecodeResult FBinkAudioInfo::Decode(const uint8* CompressedData, const int32 Co
 
 		if (RemnCompressedDataSize == 0)
 		{
-			// This is the normal termination condition
+			// This is the normal termination condition if we are routing through the output res.
 			break;
 		}
 
@@ -510,12 +481,34 @@ FDecodeResult FBinkAudioInfo::Decode(const uint8* CompressedData, const int32 Co
 			break;
 		}
 
-		uint8 const* BlockStart =0;
-		uint8 const* BlockEnd =0;
-		uint32 TrimToSampleCount =0;
-		BinkAudioCrackBlock(CompressedData, &BlockStart, &BlockEnd, &TrimToSampleCount);
+		uint8 const* BlockStart = nullptr;
+		uint8 const* BlockEnd = nullptr;
+		uint32 TrimToFrameCount = 0;
+		BinkAudioCrackBlock(CompressedData, &BlockStart, &BlockEnd, &TrimToFrameCount);
 		uint8 const* BlockBase = CompressedData;
 
+		uint32 DecodeFramesThisBlock = DecodeFrames;
+		if (TrimToFrameCount != ~0U)
+		{
+			if (TrimToFrameCount > DecodeFrames)
+			{
+				bErrorStateLatch = true;
+				return FDecodeResult(); // corrupted - should never encode a trim LARGER than the block.
+			}
+			DecodeFramesThisBlock = TrimToFrameCount;
+		}
+
+		// If we are consuming more than the entire block, just advance without actually decoding.
+		if (Decoder->ConsumeFrameCount >= DecodeFramesThisBlock)
+		{
+			uint32 InputConsumed = (uint32)(BlockEnd - CompressedData);
+			CompressedData += InputConsumed;
+			RemnCompressedDataSize -= InputConsumed;
+
+			Decoder->ConsumeFrameCount -= DecodeFramesThisBlock;
+			continue;
+		}
+
 		//
 		// We need to make sure there's room available for Bink to read past the end
 		// of the buffer (for vector decoding). If there's not, we need to copy to a
@@ -549,27 +542,30 @@ FDecodeResult FBinkAudioInfo::Decode(const uint8* CompressedData, const int32 Co
 
 		//
 		// If we're a simple single stream and we have enough output space,
-		// just decode directly in to our destination to avoid some
+		// just decode directly into our destination to avoid some
 		// copies.
 		//
-		// We also have to have a "simple" decode - i.e. aligned and no trimming.
+		// We also have to have a "simple" decode - i.e. aligned and no consume. This should be almost all
+		// of the blocks in a mono/stereo streaming source.
 		//
 		if (Decoder->StreamCount == 1 &&
-			DecodeSize <= RemnOutputPCMDataSize &&
-			TrimToSampleCount == ~0U &&
+			AllStreamsDecodeSize <= RemnOutputPCMDataSize &&
 			(((size_t)OutPCMData)&0xf) == 0 &&
 			Decoder->ConsumeFrameCount == 0)
 		{
 			uint32 DecodedBytes = BinkInterface->DecodeFn(Decoder->Decoders()[0], OutPCMData, RemnOutputPCMDataSize, &BlockStart, BlockEnd);
-			check(DecodedBytes);
-			if (DecodedBytes == 0)
+			check (DecodedBytes == AllStreamsDecodeSize);
+			if (DecodedBytes != AllStreamsDecodeSize)
 			{
 				bErrorStateLatch = true;
-
-				// This means that our block check above succeeded and we still failed - corrupted data!
-				return FDecodeResult();
+				return FDecodeResult(); // bink should always return full blocks
 			}
 
+			// Set to any trimmed value
+			DecodedBytes = DecodeFramesThisBlock * FrameSize;
+
+			check(BlockStart == BlockEnd);
+
 			uint32 InputConsumed = (uint32)(BlockStart - BlockBase);
 
 			OutPCMData += DecodedBytes;
@@ -579,103 +575,105 @@ FDecodeResult FBinkAudioInfo::Decode(const uint8* CompressedData, const int32 Co
 			continue;
 		}
 
-		// Otherwise we route through the output reservoir.
-		if (Decoder->StreamCount == 1)
-		{
-			uint32 DecodedBytes = BinkInterface->DecodeFn(Decoder->Decoders()[0], Decoder->OutputReservoir(), Decoder->OutputReservoirTotalBytes, &BlockStart, BlockEnd);
-			check(DecodedBytes);
-			if (DecodedBytes == 0)
-			{
-				bErrorStateLatch = true;
+		// Otherwise, we go into a buffer for deinterlacing / trimming / whatever.
+		// We interlace each stream so we only ever need stereo space here.
+		DeinterleavedDecodeBuffer.SetNumUninitialized(DecodeFrames * 2);
 
-				// This means that our block check above succeeded and we still failed - corrupted data!
-				return FDecodeResult();
-			}
+		// Bink always emits full transform blocks so we can do our frame trimming up
+		// front.
+		uint32 DecodedFramesStart = 0;
+		uint32 DecodedFramesEnd = DecodeFramesThisBlock;
 
-			uint32 InputConsumed = (uint32)(BlockStart - BlockBase);
+		uint32 FramesAvailable = DecodedFramesEnd - DecodedFramesStart;
 
-			CompressedData += InputConsumed;
-			RemnCompressedDataSize -= InputConsumed;
+		// Check if we need to eat some frames due to a sample-accurate seek.
+		if (Decoder->ConsumeFrameCount)
+		{
+			uint32 ConsumedFrameCount = FMath::Min((uint32)Decoder->ConsumeFrameCount, FramesAvailable);
+			DecodedFramesStart += ConsumedFrameCount;
+			FramesAvailable = DecodedFramesEnd - DecodedFramesStart;
+			Decoder->ConsumeFrameCount -= ConsumedFrameCount;
+		}
 
-			Decoder->OutputReservoirValidBytes = DecodedBytes;
+		// Check if we can deinterleave directly to the output
+		int16* InterleaveDestination = (int16*)OutPCMData;
+
+		// Since we are interlacing into this destination we only need what we use,
+		// not the full decode.
+		bool bDirectDecode = true;
+		if (FramesAvailable * FrameSize <= RemnOutputPCMDataSize)
+		{
+			// Yay - can just go direct and not even allocate the output reservoir.
 		}
 		else
 		{
-			// multistream requires interlacing the stereo/mono streams
-			int16* DeinterlaceBuffer = Decoder->DeinterlaceBuffer();
-			uint8* CurrentOutputReservoir = Decoder->OutputReservoir();
+			// We have to go through the output reservoir! We allocate enough for any run we do
+			// since it'll just stay allocated (once we use it once, we assume we'll continue to need
+			// it due to the circumstances when it's needed in UE).
+			OutputReservoir.SetNumUninitialized(DecodeFrames * FrameSize);
+			InterleaveDestination = (int16*)OutputReservoir.GetData();
+			bDirectDecode = false;
+		}
 
-			uint32 LocalNumChannels = NumChannels;
-			for (uint32 i = 0; i < Decoder->StreamCount; i++)
+		uint32 ChannelIndex = 0;
+		for (uint32 i = 0, StreamChannels = Decoder->StreamChannels()[i]; i < Decoder->StreamCount; i++, ChannelIndex += StreamChannels)
+		{
+			// Decode into the channel's slot in the deinterlace buffer.
+			uint32 DecodedBytes = BinkInterface->DecodeFn(
+				Decoder->Decoders()[i],
+				(uint8*)DeinterleavedDecodeBuffer.GetData(), 
+				DecodeFrames * sizeof(int16) * StreamChannels,
+				&BlockStart, BlockEnd);
+
+			check(DecodedBytes == DecodeFrames * StreamChannels * sizeof(int16));
+			if (DecodedBytes != DecodeFrames * StreamChannels * sizeof(int16))
 			{
-				uint32 StreamChannels = Decoder->StreamChannels()[i];
-
-				uint32 DecodedBytes = BinkInterface->DecodeFn(Decoder->Decoders()[i], (uint8*)DeinterlaceBuffer, GetMaxFrameSizeSamples() * sizeof(int16) * 2, &BlockStart, BlockEnd);
-				check(DecodedBytes);
-
-				if (DecodedBytes == 0)
-				{
-					// This means that our block check above succeeded and we still failed - corrupted data!
-					return FDecodeResult();
-				}
+				// must decode entire blocks.
+				bErrorStateLatch = true;
+				return FDecodeResult();
+			}
 
-				// deinterleave in to the output reservoir.
-				if (StreamChannels == 1)
+			// Bink audio emits pairs interlaced already, so we can copy stereo as 4 bytes
+			if (StreamChannels == 2)
+			{
+				const uint32* InBuffer = ((uint32*)DeinterleavedDecodeBuffer.GetData()) + DecodedFramesStart;
+				uint32* ChannelInterleaveDestination = (uint32*)(InterleaveDestination + ChannelIndex);
+				for (uint32 SampleIdx = 0; SampleIdx < FramesAvailable; SampleIdx++)
 				{
-					int16* Read = DeinterlaceBuffer;
-					int16* Write = (int16*)CurrentOutputReservoir;
-					uint32 Frames = DecodedBytes / sizeof(int16);
-					for (uint32 FrameIndex = 0; FrameIndex < Frames; FrameIndex++)
-					{
-						Write[LocalNumChannels * FrameIndex] = Read[FrameIndex];
-					}
+					// we know this is aligned because we pack stereo up front.
+					// not sure if there's room for SIMD here. If it's quad we could load 4 values, mask off
+					// every other and OR in... seems like that's probably worse than just copying 4 bytes.
+					// Might be able to get something with unwinding?
+					*(uint32*)(InterleaveDestination + (ChannelIndex + (SampleIdx * NumChannels))) = InBuffer[SampleIdx];
 				}
-				else
+			}
+			else
+			{
+				// This will almost never get hit - odd multichannel isn't common.
+				const int16* InBuffer = ((int16*)DeinterleavedDecodeBuffer.GetData()) + DecodedFramesStart;
+				for (uint32 SampleIdx = 0; SampleIdx < FramesAvailable; SampleIdx++)
 				{
-					// stereo int16 pairs
-					int32* Read = (int32*)DeinterlaceBuffer;
-					int16* Write = (int16*)CurrentOutputReservoir;
-					uint32 Frames = DecodedBytes / sizeof(int32);
-					for (uint32 FrameIndex = 0; FrameIndex < Frames; FrameIndex++)
-					{
-						*(int32*)(Write + LocalNumChannels * FrameIndex) = Read[FrameIndex];
-					}
+					InterleaveDestination[ChannelIndex + (SampleIdx * NumChannels)] = InBuffer[SampleIdx];
 				}
-				
-				CurrentOutputReservoir += sizeof(int16) * StreamChannels;
-
-				Decoder->OutputReservoirValidBytes += DecodedBytes;
 			}
+		}
 
-			uint32 InputConsumed = (uint32)(BlockStart - BlockBase);			
-			CompressedData += InputConsumed;
-			RemnCompressedDataSize -= InputConsumed;
-		} // end if multi stream
-
-		// Check if we are trimming the tail due to EOF
-		if (TrimToSampleCount != ~0U)
+		uint32 InputConsumed = (uint32)(BlockStart - BlockBase);
+		CompressedData += InputConsumed;
+		RemnCompressedDataSize -= InputConsumed;
+		
+		if (bDirectDecode)
 		{
-			// Ignore the tail samples by just dropping our reservoir size
-			Decoder->OutputReservoirValidBytes = TrimToSampleCount * sizeof(int16) * NumChannels;
+			// we need to update the dest pointer since we went direct.
+			OutPCMData += FramesAvailable * FrameSize;
+			RemnOutputPCMDataSize -= FramesAvailable * FrameSize;
 		}
-
-		// Check if we need to eat some frames due to a sample-accurate seek.
-		if (Decoder->ConsumeFrameCount)
+		else
 		{
-			const uint32 BytesPerFrame = sizeof(int16) * NumChannels;
-			uint32 ValidFrames = Decoder->OutputReservoirValidBytes / BytesPerFrame;
-			if (Decoder->ConsumeFrameCount < ValidFrames)
-			{
-				FMemory::Memmove(Decoder->OutputReservoir(), Decoder->OutputReservoir() + Decoder->ConsumeFrameCount * BytesPerFrame, (ValidFrames - Decoder->ConsumeFrameCount) * BytesPerFrame);
-				Decoder->OutputReservoirValidBytes -= Decoder->ConsumeFrameCount * BytesPerFrame;
-			}
-			else
-			{
-				Decoder->OutputReservoirValidBytes = 0;
-			}
-			Decoder->ConsumeFrameCount = 0;
+			Decoder->OutputReservoirTotalBytes = FramesAvailable * FrameSize;
+			Decoder->OutputReservoirReadOffset = 0;
+			// Fall through to the next loop to copy the decoded pcm data out of the reservoir.
 		}
-		// Fall through to the next loop to copy the decoded pcm data out of the reservoir.
 	} // while need output pcm data
 
 	// We get here if we filled the output buffer or not.
diff --git a/Engine/Source/Runtime/BinkAudioDecoder/Module/Public/BinkAudioInfo.h b/Engine/Source/Runtime/BinkAudioDecoder/Module/Public/BinkAudioInfo.h
--- a/Engine/Source/Runtime/BinkAudioDecoder/Module/Public/BinkAudioInfo.h
+++ b/Engine/Source/Runtime/BinkAudioDecoder/Module/Public/BinkAudioInfo.h
@@ -35,7 +35,10 @@ protected:
 	uint32 MaxCompSpaceNeeded;
 	uint32 SampleRate;
 
+	// we lazy init this because the general use case (streaming) should always be able to decode
+	// directly to the output buffer.
+	TArray<uint8, TAlignedHeapAllocator<16>> OutputReservoir;
 	struct BinkAudioDecoder* Decoder = 0;
-	uint8* RawMemory = 0;
+	TArray<uint8, TAlignedHeapAllocator<16>> RawMemory;
 	bool bErrorStateLatch = false;
 };
