Submitting new bink decoder libs to pick up the radfft init race condition fixed a while back for the encoder libs.

diff --git a/Engine/Source/Runtime/BinkAudioDecoder/Module/BinkAudioDecoder.Build.cs b/Engine/Source/Runtime/BinkAudioDecoder/Module/BinkAudioDecoder.Build.cs
--- a/Engine/Source/Runtime/BinkAudioDecoder/Module/BinkAudioDecoder.Build.cs
+++ b/Engine/Source/Runtime/BinkAudioDecoder/Module/BinkAudioDecoder.Build.cs
@@ -16,7 +16,14 @@ public class BinkAudioDecoder : ModuleRules
         }
         if (Target.Platform == UnrealTargetPlatform.Linux)
         {
-            return Path.Combine(ModuleDirectory, "..", "SDK", "BinkAudio", "Lib", "libbinka_ue_decode_lnx64_static.a");
+            if (Target.Architecture == UnrealArch.Arm64)
+            {
+                return Path.Combine(ModuleDirectory, "..", "SDK", "BinkAudio", "Lib", "libbinka_ue_decode_lnxarm64_static.a");
+            }
+            if (Target.Architecture == UnrealArch.X64)
+            {
+                return Path.Combine(ModuleDirectory, "..", "SDK", "BinkAudio", "Lib", "libbinka_ue_decode_lnx64_static.a");
+            }
         }
         if (Target.Platform == UnrealTargetPlatform.Mac)
         {
diff --git a/Engine/Source/Runtime/BinkAudioDecoder/SDK/BinkAudio/Src/binkacd.cpp b/Engine/Source/Runtime/BinkAudioDecoder/SDK/BinkAudio/Src/binkacd.cpp
--- a/Engine/Source/Runtime/BinkAudioDecoder/SDK/BinkAudio/Src/binkacd.cpp
+++ b/Engine/Source/Runtime/BinkAudioDecoder/SDK/BinkAudio/Src/binkacd.cpp
@@ -13,6 +13,15 @@
 
 #include "radfft.h"
 
+#ifndef alloca
+#if defined(_MSC_VER)
+  #include <malloc.h>
+#elif defined(__GNUC__)
+  #include <alloca.h>
+#endif
+#endif
+
+
 #ifdef BIG_OLE_FFT // always false for ue binka
 #define MAX_TRANSFORM                 4096
 #else
@@ -112,7 +121,7 @@ typedef struct BINKAUDIODECOMP
 
 #include "binkbits.h"
 
-#if  defined(__RADARM64__)
+#if  defined(__RADARM64__) && !defined(_M_ARM64) // exclude visual studio arm64 as it doesn't support __int128
   #define bigreg __int128
   #define bigregzero 0
   #define bigregloadaligned( val, ptr ) (val)=*((bigreg*)(ptr))
diff --git a/Engine/Source/Runtime/BinkAudioDecoder/SDK/BinkAudio/Src/console.h b/Engine/Source/Runtime/BinkAudioDecoder/SDK/BinkAudio/Src/console.h
--- a/Engine/Source/Runtime/BinkAudioDecoder/SDK/BinkAudio/Src/console.h
+++ b/Engine/Source/Runtime/BinkAudioDecoder/SDK/BinkAudio/Src/console.h
@@ -41,7 +41,7 @@
         static void console_handleevents() {}
         #define CONSOLE_MAIN extern "C" int main(int argc, char** argv)
         
-    #elif defined(__RADNT__)
+    #elif defined(__RADNT__) || defined(__RADWINRT__)
         #include <conio.h>
         #include <windows.h>
         #include <stdio.h>
diff --git a/Engine/Source/Runtime/BinkAudioDecoder/SDK/BinkAudio/Src/egttypes.h b/Engine/Source/Runtime/BinkAudioDecoder/SDK/BinkAudio/Src/egttypes.h
--- a/Engine/Source/Runtime/BinkAudioDecoder/SDK/BinkAudio/Src/egttypes.h
+++ b/Engine/Source/Runtime/BinkAudioDecoder/SDK/BinkAudio/Src/egttypes.h
@@ -173,7 +173,7 @@
 
 
 #if !defined(__RADDETECTED__)
-  #error "radtypes.h did not detect your platform."
+  #error "egttypes.h did not detect your platform."
 #endif
 
 // ========================================================
@@ -216,7 +216,7 @@
   #define __RADDETECTEDPROC__ __RADPPC__
   #define __RADBIGENDIAN__
 #endif
-#if defined( __aarch64__ ) || defined( __arm64__ )
+#if defined( __aarch64__ ) || defined( __arm64__ ) || defined(_M_ARM64)
   #define __RADARM__ 1
   #define __RADARM64__ 6
   #define __RADDETECTEDPROC__ __RADARM64__
@@ -232,7 +232,7 @@
   #error "egttypes.h did not detect your processor type."
 #endif
 
-#if defined(__ppc64__) || defined(__aarch64__) || defined(_M_X64) || defined(__x86_64__) || defined(__x86_64)
+#if defined(__ppc64__) || defined(__aarch64__) || defined(_M_X64) || defined(__x86_64__) || defined(__x86_64) || defined(_M_ARM64)
   #define __RAD64__
   #define __RAD64REGS__  // need to set this for platforms that aren't 64-bit, but have 64-bit regs (old consoles)
 #endif
diff --git a/Engine/Source/Runtime/BinkAudioDecoder/SDK/BinkAudio/Src/radfft.cpp b/Engine/Source/Runtime/BinkAudioDecoder/SDK/BinkAudio/Src/radfft.cpp
--- a/Engine/Source/Runtime/BinkAudioDecoder/SDK/BinkAudio/Src/radfft.cpp
+++ b/Engine/Source/Runtime/BinkAudioDecoder/SDK/BinkAudio/Src/radfft.cpp
@@ -2104,20 +2104,23 @@ static void init_tables()
 
 #endif
 
-void RADLINK radfft_init()
+static int radfft_init_helper()
 {
-  static int done_init = 0;
-  if ( done_init == 0 )
-  {
-      done_init = 1;
 #ifdef CHOOSE_KERNELS
-      s_kernel = CHOOSE_KERNELS;
+    s_kernel = CHOOSE_KERNELS;
 #endif
 
 #ifndef USETABLES
-      init_tables();
+    init_tables();
 #endif
-  }
+    return 1;
+}
+
+void RADLINK radfft_init()
+{
+    // lean on c++ static init here: this is guaranteed by c++ to be called once
+    // and will be wrapped with a mutex by the compiler.
+    static int done_init = radfft_init_helper();
 }
 
 // Complex FFT
