Removing check() in favor of error handling to unblock - needs further investigation!

diff --git a/Engine/Source/Runtime/BinkAudioDecoder/Module/Private/BinkAudioInfo.cpp b/Engine/Source/Runtime/BinkAudioDecoder/Module/Private/BinkAudioInfo.cpp
--- a/Engine/Source/Runtime/BinkAudioDecoder/Module/Private/BinkAudioInfo.cpp
+++ b/Engine/Source/Runtime/BinkAudioDecoder/Module/Private/BinkAudioInfo.cpp
@@ -564,7 +564,14 @@ FDecodeResult FBinkAudioInfo::Decode(const uint8* CompressedData, const int32 Co
 			// Set to any trimmed value
 			DecodedBytes = DecodeFramesThisBlock * FrameSize;
 
-			check(BlockStart == BlockEnd);
+			if (BlockStart != BlockEnd)
+			{
+				// Header mismatch? We should always consume exactly what we expected to.
+				UE_LOG(LogBinkAudioDecoder, Error, TEXT("BinkAudio consumed unexpected amount! BlockEnd = 0x%llx BlockStart = 0x%llx BlockBase = 0x%llx"),
+					(uint64)BlockEnd, (uint64)BlockStart, (uint64)BlockBase);
+				bErrorStateLatch = true;
+				return FDecodeResult();
+			}
 
 			uint32 InputConsumed = (uint32)(BlockStart - BlockBase);
 
