Add error logs to BinkAudioInfo for failure cases when parsing the header.

#tests PIE with break point to trigger failure case

diff --git a/Engine/Source/Runtime/BinkAudioDecoder/Module/Private/BinkAudioInfo.cpp b/Engine/Source/Runtime/BinkAudioDecoder/Module/Private/BinkAudioInfo.cpp
--- a/Engine/Source/Runtime/BinkAudioDecoder/Module/Private/BinkAudioInfo.cpp
+++ b/Engine/Source/Runtime/BinkAudioDecoder/Module/Private/BinkAudioInfo.cpp
@@ -26,6 +26,29 @@ DEFINE_LOG_CATEGORY_STATIC(LogBinkAudioDecoder, Log, All);
 // it needs to be < 255
 #define MAX_BINK_AUDIO_CHANNELS 16
 
+namespace BinkAudioInfo
+{
+	static void LogOnce(BinkAudioFileHeader& Header)
+	{
+		static std::atomic<bool> bLogHeader = true;
+		if (bLogHeader.exchange(false))
+		{
+			UE_LOG(LogBinkAudioDecoder, Log, TEXT("FBinkAudioInfo::ParseHeader: BEGIN LOG HEADER"));
+			UE_LOG(LogBinkAudioDecoder, Log, TEXT("Header.tag = %d"), Header.tag);
+			UE_LOG(LogBinkAudioDecoder, Log, TEXT("Header.version = %d"), Header.version);
+			UE_LOG(LogBinkAudioDecoder, Log, TEXT("Header.channels = %d"), Header.channels);
+			UE_LOG(LogBinkAudioDecoder, Log, TEXT("Header.PADDING = %d"), Header.PADDING);
+			UE_LOG(LogBinkAudioDecoder, Log, TEXT("Header.rate = %d"), Header.rate);
+			UE_LOG(LogBinkAudioDecoder, Log, TEXT("Header.sample_count = %d"), Header.sample_count);
+			UE_LOG(LogBinkAudioDecoder, Log, TEXT("Header.max_comp_space_needed = %d"), Header.max_comp_space_needed);
+			UE_LOG(LogBinkAudioDecoder, Log, TEXT("Header.flags = %d"), Header.flags);
+			UE_LOG(LogBinkAudioDecoder, Log, TEXT("Header.output_file_size = %d"), Header.seek_table_entry_count);
+			UE_LOG(LogBinkAudioDecoder, Log, TEXT("Header.blocks_per_seek_table_entry = %d"), Header.blocks_per_seek_table_entry);
+			UE_LOG(LogBinkAudioDecoder, Log, TEXT("FBinkAudioInfo::ParseHeader: END LOG HEADER"));
+		}
+	}
+}
+
 //-----------------------------------------------------------------------------
 //
 // All memory for the decoder is in one contiguous block:
@@ -220,16 +243,21 @@ bool FBinkAudioInfo::ParseHeader(const uint8* InSrcBufferData, uint32 InSrcBuffe
 	check(InSrcBufferDataSize >= sizeof(BinkAudioFileHeader));
 	if (InSrcBufferDataSize < sizeof(BinkAudioFileHeader))
 	{
+		UE_LOG(LogBinkAudioDecoder, Error, TEXT("FBinkAudioInfo::ParseHeader: Failed: InSrcBufferDataSize < sizeof(BinkAudioFileHeader) : %d < %llu"), InSrcBufferDataSize, sizeof(BinkAudioFileHeader));
 		return false;
 	}
 
 	BinkAudioFileHeader* Header = (BinkAudioFileHeader*)InSrcBufferData;
 	if (Header->tag != 'UEBA')
 	{
+		UE_LOG(LogBinkAudioDecoder, Error, TEXT("FBinkAudioInfo::ParseHeader: Failed: Header->tag != \'UEBA\' (%d): tag = %d"), (uint32)('UEBA'), Header->tag);
+		BinkAudioInfo::LogOnce(*Header);
 		return false;
 	}
 	if (Header->version != 1)
 	{
+		UE_LOG(LogBinkAudioDecoder, Error, TEXT("FBinkAudioInfo::ParseHeader: Failed: Header->version != 1: version = %d"), Header->version);
+		BinkAudioInfo::LogOnce(*Header);
 		return false;
 	}
 
@@ -242,6 +270,8 @@ bool FBinkAudioInfo::ParseHeader(const uint8* InSrcBufferData, uint32 InSrcBuffe
 	uint32 SeekTableSize = Header->seek_table_entry_count * sizeof(uint16);
 	if (sizeof(BinkAudioFileHeader) + SeekTableSize > InSrcBufferDataSize)
 	{
+		UE_LOG(LogBinkAudioDecoder, Error, TEXT("BinkAudioInfo::ParseHeader Failed: Invalid SeekTableSize: SeekTableSize = %d; sizeof(BinkAudioFileHeader) = %llu; InSrcBufferDataSize = %d"), SeekTableSize, SeekTableSize + sizeof(BinkAudioFileHeader), InSrcBufferDataSize);
+		BinkAudioInfo::LogOnce(*Header);
 		return false;
 	}
 
@@ -272,6 +302,7 @@ bool FBinkAudioInfo::CreateDecoder()
 
 	if (BinkInterface == nullptr)
 	{
+		UE_LOG(LogBinkAudioDecoder, Error, TEXT("BinkAudioInfo::CreateDecoder Failed: BinkInterface does not exist for the current configuration"));
 		return false; // only happens if we dont have libs.
 	}
 
