Ensure visionOS simulator BinkAudioDecoder static libs are built against simulator target.

diff --git a/Engine/Source/Runtime/BinkAudioDecoder/SDK/BinkAudio/BuildDecoderVisionOS.command b/Engine/Source/Runtime/BinkAudioDecoder/SDK/BinkAudio/BuildDecoderVisionOS.command
--- a/Engine/Source/Runtime/BinkAudioDecoder/SDK/BinkAudio/BuildDecoderVisionOS.command
+++ b/Engine/Source/Runtime/BinkAudioDecoder/SDK/BinkAudio/BuildDecoderVisionOS.command
@@ -23,7 +23,7 @@ mkdir $BUILD_LOCATION/visionos/sim
 
 SIM_ARGS_ARM64=(
 	-arch arm64
-	--target=arm-apple-xros1
+	--target=arm-apple-xros1-simulator
 	-isysroot /Applications/Xcode.app/Contents/Developer/Platforms/XRSimulator.platform/Developer/SDKs/XRSimulator.sdk
 )
 
