bink audio encoder lib for arm64

#rnx

diff --git a/Engine/Source/Runtime/BinkAudioDecoder/SDK/BinkAudio/Src/varbits.h b/Engine/Source/Runtime/BinkAudioDecoder/SDK/BinkAudio/Src/varbits.h
--- a/Engine/Source/Runtime/BinkAudioDecoder/SDK/BinkAudio/Src/varbits.h
+++ b/Engine/Source/Runtime/BinkAudioDecoder/SDK/BinkAudio/Src/varbits.h
@@ -171,7 +171,15 @@ RADDEFFUNC void VarBitsCopy(VARBITS* dest,VARBITS* src,U32 size);
 
   #ifdef _MSC_VER
 
-    #define getbitlevelvar(n) (U32) (32 - _arm_clz(n))
+    #ifdef _M_ARM64
+
+        #define getbitlevelvar(n) (U32) (32 - _CountLeadingZeros(n)) 
+
+    #else
+
+       #define getbitlevelvar(n) (U32) (32 - _arm_clz(n))
+
+    #endif
 
   #else
 
