PR #14256: Fix for Bink Audio Decoder not working on LinuxArm64 platforms

#tests LowLevelTests AudioUnitTests

#Synced-CL 49865182

diff --git a/Engine/Source/Runtime/BinkAudioDecoder/Module/BinkAudioDecoder.Build.cs b/Engine/Source/Runtime/BinkAudioDecoder/Module/BinkAudioDecoder.Build.cs
--- a/Engine/Source/Runtime/BinkAudioDecoder/Module/BinkAudioDecoder.Build.cs
+++ b/Engine/Source/Runtime/BinkAudioDecoder/Module/BinkAudioDecoder.Build.cs
@@ -18,7 +18,7 @@ public class BinkAudioDecoder : ModuleRules
 			}
 			return Path.Combine(ModuleDirectory, "..", "SDK", "BinkAudio", "Lib", "binka_ue_decode_win64_static.lib");
 		}
-		if (Target.Platform == UnrealTargetPlatform.Linux)
+		if (Target.Platform.IsInGroup(UnrealPlatformGroup.Linux))
         {
             if (Target.Architecture == UnrealArch.Arm64)
             {
